PS4 - Change the shader sdb output folder and add the shader debug mode settings to the shader key

diff --git a/Engine/Source/Programs/ShaderCompileWorker/Private/ShaderCompileWorker.cpp b/Engine/Source/Programs/ShaderCompileWorker/Private/ShaderCompileWorker.cpp
--- a/Engine/Source/Programs/ShaderCompileWorker/Private/ShaderCompileWorker.cpp
+++ b/Engine/Source/Programs/ShaderCompileWorker/Private/ShaderCompileWorker.cpp
@@ -13,7 +13,7 @@
 
 #define DEBUG_USING_CONSOLE	0
 
-const int32 ShaderCompileWorkerInputVersion = 1;
+const int32 ShaderCompileWorkerInputVersion = 2;
 const int32 ShaderCompileWorkerOutputVersion = 1;
 
 double LastCompileTime = 0.0;
diff --git a/Engine/Source/Runtime/Core/Private/HAL/ConsoleManager.cpp b/Engine/Source/Runtime/Core/Private/HAL/ConsoleManager.cpp
--- a/Engine/Source/Runtime/Core/Private/HAL/ConsoleManager.cpp
+++ b/Engine/Source/Runtime/Core/Private/HAL/ConsoleManager.cpp
@@ -1998,6 +1998,20 @@ void CreateConsoleVariables()
 		TEXT(" 1: on (Default)"),
 		ECVF_ReadOnly);
 
+	IConsoleManager::Get().RegisterConsoleVariable(TEXT("r.PS4MixedModeShaderDebugInfo"),
+		0,
+		TEXT("Whether to compile shaders to allow mixed mode shader debugging. This will currently generate slower code.\n")
+		TEXT(" 0: Normal mode\n")
+		TEXT(" 1: Mixed mode)"),
+		ECVF_ReadOnly);
+
+	IConsoleManager::Get().RegisterConsoleVariable(TEXT("r.PS4DumpShaderSDB"),
+		0,
+		TEXT("Whether to dump shader sdb files used for shader association.\n")
+		TEXT(" 0: Disabled\n")
+		TEXT(" 1: Enabled)"),
+		ECVF_ReadOnly);
+
 	IConsoleManager::Get().RegisterConsoleVariable(TEXT("r.ViewDistanceScale"),
 		1.0f,
 		TEXT("Controls the view distance scale. A primitive's MaxDrawDistance is scaled by this value.\n")
diff --git a/Engine/Source/Runtime/Engine/Private/GlobalShader.cpp b/Engine/Source/Runtime/Engine/Private/GlobalShader.cpp
--- a/Engine/Source/Runtime/Engine/Private/GlobalShader.cpp
+++ b/Engine/Source/Runtime/Engine/Private/GlobalShader.cpp
@@ -292,7 +292,7 @@ FString GetGlobalShaderMapKeyString(const FGlobalShaderMapId& ShaderMapId, EShad
 {
 	FName Format = LegacyShaderPlatformToShaderFormat(Platform);
 	FString ShaderMapKeyString = Format.ToString() + TEXT("_") + FString(FString::FromInt(GetTargetPlatformManagerRef().ShaderFormatVersion(Format))) + TEXT("_");
-	ShaderMapAppendKeyString(ShaderMapKeyString);
+	ShaderMapAppendKeyString(Platform, ShaderMapKeyString);
 	ShaderMapId.AppendKeyString(ShaderMapKeyString);
 	return FDerivedDataCacheInterface::BuildCacheKey(TEXT("GSM"), GLOBALSHADERMAP_DERIVEDDATA_VER, *ShaderMapKeyString);
 }
diff --git a/Engine/Source/Runtime/Engine/Private/Materials/MaterialShader.cpp b/Engine/Source/Runtime/Engine/Private/Materials/MaterialShader.cpp
--- a/Engine/Source/Runtime/Engine/Private/Materials/MaterialShader.cpp
+++ b/Engine/Source/Runtime/Engine/Private/Materials/MaterialShader.cpp
@@ -746,7 +746,7 @@ FString GetMaterialShaderMapKeyString(const FMaterialShaderMapId& ShaderMapId, E
 {
 	FName Format = LegacyShaderPlatformToShaderFormat(Platform);
 	FString ShaderMapKeyString = Format.ToString() + TEXT("_") + FString(FString::FromInt(GetTargetPlatformManagerRef().ShaderFormatVersion(Format))) + TEXT("_");
-	ShaderMapAppendKeyString(ShaderMapKeyString);
+	ShaderMapAppendKeyString(Platform, ShaderMapKeyString);
 	ShaderMapId.AppendKeyString(ShaderMapKeyString);
 	return FDerivedDataCacheInterface::BuildCacheKey(TEXT("MATSM"), MATERIALSHADERMAP_DERIVEDDATA_VER, *ShaderMapKeyString);
 }
diff --git a/Engine/Source/Runtime/Engine/Private/ShaderCompiler/ShaderCompiler.cpp b/Engine/Source/Runtime/Engine/Private/ShaderCompiler/ShaderCompiler.cpp
--- a/Engine/Source/Runtime/Engine/Private/ShaderCompiler/ShaderCompiler.cpp
+++ b/Engine/Source/Runtime/Engine/Private/ShaderCompiler/ShaderCompiler.cpp
@@ -32,7 +32,7 @@ static FAutoConsoleVariableRef CVarDumpShaderDebugInfo(
 // Serialize Queued Job information
 static void DoWriteTasks(TArray<FShaderCompileJob*>& QueuedJobs, FArchive& TransferFile)
 {
-	int32 ShaderCompileWorkerInputVersion = 1;
+	int32 ShaderCompileWorkerInputVersion = 2;
 	TransferFile << ShaderCompileWorkerInputVersion;
 	int32 NumBatches = QueuedJobs.Num();
 	TransferFile << NumBatches;
@@ -1930,6 +1930,7 @@ void GlobalBeginCompileShader(
 	Input.ShaderFormat = LegacyShaderPlatformToShaderFormat(EShaderPlatform(Target.Platform));
 	Input.SourceFilename = SourceFilename;
 	Input.EntryPointName = FunctionName;
+	Input.ShaderDebugInfoDirectory = GShaderCompilingManager->GetAbsoluteShaderDebugInfoDirectory();
 
 	if (GDumpShaderDebugInfo != 0)
 	{
diff --git a/Engine/Source/Runtime/ShaderCore/Private/Shader.cpp b/Engine/Source/Runtime/ShaderCore/Private/Shader.cpp
--- a/Engine/Source/Runtime/ShaderCore/Private/Shader.cpp
+++ b/Engine/Source/Runtime/ShaderCore/Private/Shader.cpp
@@ -1090,7 +1090,7 @@ FName GetRuntimeShaderFormat()
 	return RuntimeFormat;
 }
 
-void ShaderMapAppendKeyString(FString& KeyString)
+void ShaderMapAppendKeyString(EShaderPlatform Platform, FString& KeyString)
 {
 	// Globals that should cause all shaders to recompile when changed must be appended to the key here
 	// Key should be kept as short as possible while being somewhat human readable for debugging
@@ -1120,4 +1120,23 @@ void ShaderMapAppendKeyString(FString& KeyString)
 		static IConsoleVariable* CVar = IConsoleManager::Get().FindConsoleVariable(TEXT("r.DBuffer"));
 		KeyString += (CVar && CVar->GetInt() != 0) ? TEXT("_DBuf") : TEXT("_NoDBuf");
 	}
+
+	if( Platform == SP_PS4 )
+	{
+		{
+			static const auto CVar = IConsoleManager::Get().FindTConsoleVariableDataInt(TEXT("r.PS4MixedModeShaderDebugInfo"));
+			if( CVar && CVar->GetValueOnGameThread() != 0 )
+			{
+				KeyString += TEXT("_MMDBG");
+			}
+		}
+
+		{
+			static const auto CVar = IConsoleManager::Get().FindTConsoleVariableDataInt(TEXT("r.PS4DumpShaderSDB"));
+			if( CVar && CVar->GetValueOnGameThread() != 0 )
+			{
+				KeyString += TEXT("_SDB");
+			}
+		}
+	}
 }
diff --git a/Engine/Source/Runtime/ShaderCore/Public/Shader.h b/Engine/Source/Runtime/ShaderCore/Public/Shader.h
--- a/Engine/Source/Runtime/ShaderCore/Public/Shader.h
+++ b/Engine/Source/Runtime/ShaderCore/Public/Shader.h
@@ -1338,4 +1338,4 @@ extern SHADERCORE_API const TArray<FName>& GetTargetShaderFormats();
 extern SHADERCORE_API FName GetRuntimeShaderFormat();
 
 /** Appends to KeyString for all shaders. */
-extern SHADERCORE_API void ShaderMapAppendKeyString(FString& KeyString);
\ No newline at end of file
+extern SHADERCORE_API void ShaderMapAppendKeyString(EShaderPlatform Platform, FString& KeyString);
\ No newline at end of file
diff --git a/Engine/Source/Runtime/ShaderCore/Public/ShaderCore.h b/Engine/Source/Runtime/ShaderCore/Public/ShaderCore.h
--- a/Engine/Source/Runtime/ShaderCore/Public/ShaderCore.h
+++ b/Engine/Source/Runtime/ShaderCore/Public/ShaderCore.h
@@ -328,6 +328,7 @@ struct FShaderCompilerInput
 	FString SourceFilename;
 	FString EntryPointName;
 	FString DumpDebugInfoPath;
+	FString ShaderDebugInfoDirectory;
 	FShaderCompilerEnvironment Environment;
 	TRefCountPtr<FShaderCompilerEnvironment> SharedEnvironment;
 
@@ -345,6 +346,7 @@ struct FShaderCompilerInput
 		Ar << Input.SourceFilename;
 		Ar << Input.EntryPointName;
 		Ar << Input.DumpDebugInfoPath;
+		Ar << Input.ShaderDebugInfoDirectory;
 		Ar << Input.Environment;
 
 		bool bHasSharedEnvironment = IsValidRef(Input.SharedEnvironment);
