UE4: Rearrange RHICmdList toward non-null cmd lists and the final "look"

diff --git a/Engine/Source/Runtime/Renderer/Private/ShaderBaseClasses.h b/Engine/Source/Runtime/Renderer/Private/ShaderBaseClasses.h
--- a/Engine/Source/Runtime/Renderer/Private/ShaderBaseClasses.h
+++ b/Engine/Source/Runtime/Renderer/Private/ShaderBaseClasses.h
@@ -119,7 +119,7 @@ public:
 	FUniformBufferRHIParamRef GetParameterCollectionBuffer(const FGuid& Id, const FSceneInterface* SceneInterface) const;
 
 	template<typename ShaderRHIParamRef>
-	void SetParameters(FRHICommandList* RHICmdList, const ShaderRHIParamRef ShaderRHI,const FSceneView& View)
+	void SetParameters(FRHICommandList& RHICmdList, const ShaderRHIParamRef ShaderRHI,const FSceneView& View)
 	{
 		check(GetUniformBufferParameter<FViewUniformShaderParameters>().IsInitialized());
 		CheckShaderIsValid();
@@ -134,7 +134,7 @@ public:
 	/** Sets pixel parameters that are material specific but not FMeshBatch specific. */
 	template<typename ShaderRHIParamRef>
 	void SetParameters(
-		FRHICommandList* RHICmdList,
+		FRHICommandList& RHICmdList,
 		const ShaderRHIParamRef ShaderRHI, 
 		const FMaterialRenderProxy* MaterialRenderProxy, 
 		const FMaterial& Material,
@@ -392,18 +392,18 @@ public:
 	{}
 
 	template<typename ShaderRHIParamRef>
-	void SetParameters(FRHICommandList* RHICmdList, const ShaderRHIParamRef ShaderRHI,const FMaterialRenderProxy* MaterialRenderProxy,const FMaterial& Material,const FSceneView& View,ESceneRenderTargetsMode::Type TextureMode)
+	void SetParameters(FRHICommandList& RHICmdList, const ShaderRHIParamRef ShaderRHI,const FMaterialRenderProxy* MaterialRenderProxy,const FMaterial& Material,const FSceneView& View,ESceneRenderTargetsMode::Type TextureMode)
 	{
 		FMaterialShader::SetParameters(RHICmdList, ShaderRHI,MaterialRenderProxy,Material,View,false,TextureMode);
 	}
 
-	void SetVFParametersOnly(FRHICommandList* RHICmdList, const FVertexFactory* VertexFactory,const FSceneView& View,const FMeshBatchElement& BatchElement)
+	void SetVFParametersOnly(FRHICommandList& RHICmdList, const FVertexFactory* VertexFactory,const FSceneView& View,const FMeshBatchElement& BatchElement)
 	{
 		VertexFactoryParameters.SetMesh(RHICmdList, this,VertexFactory,View,BatchElement, 0);
 	}
 
 	template<typename ShaderRHIParamRef>
-	void SetMesh(FRHICommandList* RHICmdList, const ShaderRHIParamRef ShaderRHI,const FVertexFactory* VertexFactory,const FSceneView& View,const FPrimitiveSceneProxy* Proxy,const FMeshBatchElement& BatchElement,uint32 DataFlags = 0)
+	void SetMesh(FRHICommandList& RHICmdList, const ShaderRHIParamRef ShaderRHI,const FVertexFactory* VertexFactory,const FSceneView& View,const FPrimitiveSceneProxy* Proxy,const FMeshBatchElement& BatchElement,uint32 DataFlags = 0)
 	{
 		// Set the mesh for the vertex factory
 		VertexFactoryParameters.SetMesh(RHICmdList, this,VertexFactory,View,BatchElement, DataFlags);
@@ -476,12 +476,12 @@ public:
 
 	FBaseHS() {}
 
-	void SetParameters(FRHICommandList* RHICmdList, const FMaterialRenderProxy* MaterialRenderProxy,const FSceneView& View)
+	void SetParameters(FRHICommandList& RHICmdList, const FMaterialRenderProxy* MaterialRenderProxy,const FSceneView& View)
 	{
 		FMeshMaterialShader::SetParameters(RHICmdList, GetHullShader(), MaterialRenderProxy, *MaterialRenderProxy->GetMaterial(View.GetFeatureLevel()), View, ESceneRenderTargetsMode::SetTextures);
 	}
 
-	void SetMesh(FRHICommandList* RHICmdList, const FVertexFactory* VertexFactory,const FSceneView& View,const FPrimitiveSceneProxy* Proxy,const FMeshBatchElement& BatchElement)
+	void SetMesh(FRHICommandList& RHICmdList, const FVertexFactory* VertexFactory,const FSceneView& View,const FPrimitiveSceneProxy* Proxy,const FMeshBatchElement& BatchElement)
 	{
 		FMeshMaterialShader::SetMesh(RHICmdList, GetHullShader(),VertexFactory,View,Proxy,BatchElement);
 	}
@@ -522,12 +522,12 @@ public:
 
 	FBaseDS() {}
 
-	void SetParameters(FRHICommandList* RHICmdList, const FMaterialRenderProxy* MaterialRenderProxy,const FSceneView& View)
+	void SetParameters(FRHICommandList& RHICmdList, const FMaterialRenderProxy* MaterialRenderProxy,const FSceneView& View)
 	{
 		FMeshMaterialShader::SetParameters(RHICmdList, GetDomainShader(), MaterialRenderProxy, *MaterialRenderProxy->GetMaterial(View.GetFeatureLevel()), View, ESceneRenderTargetsMode::SetTextures);
 	}
 
-	void SetMesh(FRHICommandList* RHICmdList, const FVertexFactory* VertexFactory,const FSceneView& View,const FPrimitiveSceneProxy* Proxy,const FMeshBatchElement& BatchElement)
+	void SetMesh(FRHICommandList& RHICmdList, const FVertexFactory* VertexFactory,const FSceneView& View,const FPrimitiveSceneProxy* Proxy,const FMeshBatchElement& BatchElement)
 	{
 		FMeshMaterialShader::SetMesh(RHICmdList, GetDomainShader(),VertexFactory,View,Proxy,BatchElement);
 	}
diff --git a/Engine/Source/Runtime/Renderer/Private/ShaderComplexityRendering.cpp b/Engine/Source/Runtime/Renderer/Private/ShaderComplexityRendering.cpp
--- a/Engine/Source/Runtime/Renderer/Private/ShaderComplexityRendering.cpp
+++ b/Engine/Source/Runtime/Renderer/Private/ShaderComplexityRendering.cpp
@@ -18,7 +18,7 @@ float GetMaxShaderComplexityCount(ERHIFeatureLevel::Type InFeatureType)
 }
 
 void FShaderComplexityAccumulatePS::SetParameters(
-	FRHICommandList* RHICmdList, 
+	FRHICommandList& RHICmdList, 
 	uint32 NumVertexInstructions, 
 	uint32 NumPixelInstructions,
 	ERHIFeatureLevel::Type InFeatureLevel)
@@ -88,7 +88,7 @@ public:
 	}
 
 	virtual void SetParameters(
-		FRHICommandList* RHICmdList, 
+		FRHICommandList& RHICmdList, 
 		const FRenderingCompositePassContext& Context,
 		const TArray<FLinearColor>& Colors
 		)
@@ -169,7 +169,7 @@ void FRCPassPostProcessVisualizeComplexity::Process(FRenderingCompositePassConte
 	const FSceneRenderTargetItem& DestRenderTarget = PassOutputs[0].RequestSurface(Context);
 
 	//@todo-rco: RHIPacketList
-	FRHICommandList* RHICmdList = nullptr;
+	FRHICommandList& RHICmdList = FRHICommandList::GetNullRef();
 
 	// Set the view family's render target/viewport.
 	RHISetRenderTarget(DestRenderTarget.TargetableTexture, FTextureRHIRef());	
diff --git a/Engine/Source/Runtime/Renderer/Private/ShadowRendering.cpp b/Engine/Source/Runtime/Renderer/Private/ShadowRendering.cpp
--- a/Engine/Source/Runtime/Renderer/Private/ShadowRendering.cpp
+++ b/Engine/Source/Runtime/Renderer/Private/ShadowRendering.cpp
@@ -103,7 +103,7 @@ public:
 	}
 
 	void SetParameters(
-		FRHICommandList* RHICmdList, 
+		FRHICommandList& RHICmdList, 
 		const FMaterialRenderProxy* MaterialRenderProxy,
 		const FMaterial& Material,
 		const FSceneView& View,
@@ -114,7 +114,7 @@ public:
 		ShadowParameters.SetVertexShader(RHICmdList, this, View, ShadowInfo, MaterialRenderProxy);
 	}
 
-	void SetMesh(FRHICommandList* RHICmdList, const FVertexFactory* VertexFactory,const FSceneView& View,const FPrimitiveSceneProxy* Proxy,const FMeshBatchElement& BatchElement)
+	void SetMesh(FRHICommandList& RHICmdList, const FVertexFactory* VertexFactory,const FSceneView& View,const FPrimitiveSceneProxy* Proxy,const FMeshBatchElement& BatchElement)
 	{
 		FMeshMaterialShader::SetMesh(RHICmdList, GetVertexShader(),VertexFactory,View,Proxy,BatchElement);
 	}
@@ -239,7 +239,7 @@ public:
 	}
 
 	void SetParameters(
-		FRHICommandList* RHICmdList, 
+		FRHICommandList& RHICmdList, 
 		const FMaterialRenderProxy* MaterialRenderProxy,
 		const FSceneView& View,
 		const FProjectedShadowInfo* ShadowInfo
@@ -317,7 +317,7 @@ public:
 	}
 
 	void SetParameters(
-		FRHICommandList* RHICmdList, 
+		FRHICommandList& RHICmdList, 
 		const FSceneView& View,
 		const FProjectedShadowInfo* ShadowInfo
 		)
@@ -343,7 +343,7 @@ public:
 			);
 	}
 
-	void SetMesh(FRHICommandList* RHICmdList, const FPrimitiveSceneProxy* PrimitiveSceneProxy, const FProjectedShadowInfo* ShadowInfo, const FSceneView& View)
+	void SetMesh(FRHICommandList& RHICmdList, const FPrimitiveSceneProxy* PrimitiveSceneProxy, const FProjectedShadowInfo* ShadowInfo, const FSceneView& View)
 	{
 		if (MeshVisibleToFace.IsBound())
 		{
@@ -424,7 +424,7 @@ public:
 	TShadowDepthBasePS() {}
 
 	void SetParameters(
-		FRHICommandList* RHICmdList, 
+		FRHICommandList& RHICmdList, 
 		const FMaterialRenderProxy* MaterialRenderProxy,
 		const FMaterial& Material,
 		const FSceneView& View,
@@ -462,7 +462,7 @@ public:
 		}
 	}
 
-	void SetMesh(FRHICommandList* RHICmdList, const FVertexFactory* VertexFactory,const FSceneView& View,const FPrimitiveSceneProxy* Proxy,const FMeshBatchElement& BatchElement)
+	void SetMesh(FRHICommandList& RHICmdList, const FVertexFactory* VertexFactory,const FSceneView& View,const FPrimitiveSceneProxy* Proxy,const FMeshBatchElement& BatchElement)
 	{
 		FMeshMaterialShader::SetMesh(RHICmdList, GetPixelShader(),VertexFactory,View,Proxy,BatchElement);
 	}
@@ -592,7 +592,7 @@ bool FShadowProjectionVS::ShouldCache(EShaderPlatform Platform)
 	return IsFeatureLevelSupported(Platform, ERHIFeatureLevel::SM3);
 }
 
-void FShadowProjectionVS::SetParameters(FRHICommandList* RHICmdList, const FSceneView& View, const FProjectedShadowInfo* ShadowInfo, const int32 ShadowSplitIndex)
+void FShadowProjectionVS::SetParameters(FRHICommandList& RHICmdList, const FSceneView& View, const FProjectedShadowInfo* ShadowInfo, const int32 ShadowSplitIndex)
 {
 	FGlobalShader::SetParameters(RHICmdList, GetVertexShader(),View);
 	
@@ -846,11 +846,11 @@ FShadowDepthDrawingPolicy<bRenderingReflectiveShadowMaps>::FShadowDepthDrawingPo
 }
 
 template <bool bRenderingReflectiveShadowMaps>
-void FShadowDepthDrawingPolicy<bRenderingReflectiveShadowMaps>::DrawShared(FRHICommandList* RHICmdList, const FSceneView* View,FBoundShaderStateRHIParamRef BoundShaderState) const
+void FShadowDepthDrawingPolicy<bRenderingReflectiveShadowMaps>::DrawShared(FRHICommandList& RHICmdList, const FSceneView* View,FBoundShaderStateRHIParamRef BoundShaderState) const
 {
 	checkSlow(bDirectionalLight == PolicyShadowInfo->bDirectionalLight && bPreShadow == PolicyShadowInfo->bPreShadow);
 
-	check(!RHICmdList);
+	RHICmdList.CheckIsNull();
 	// Set the actual shader & vertex declaration state
 	RHISetBoundShaderState(BoundShaderState);
 
@@ -927,7 +927,7 @@ FBoundShaderStateRHIRef FShadowDepthDrawingPolicy<bRenderingReflectiveShadowMaps
 
 template <bool bRenderingReflectiveShadowMaps>
 void FShadowDepthDrawingPolicy<bRenderingReflectiveShadowMaps>::SetMeshRenderState(
-	FRHICommandList* RHICmdList, 
+	FRHICommandList& RHICmdList, 
 	const FSceneView& View,
 	const FPrimitiveSceneProxy* PrimitiveSceneProxy,
 	const FMeshBatch& Mesh,
@@ -1035,7 +1035,7 @@ void FShadowDepthDrawingPolicyFactory::AddStaticMesh(FScene* Scene,FStaticMesh*
 }
 
 bool FShadowDepthDrawingPolicyFactory::DrawDynamicMesh(
-	FRHICommandList* RHICmdList, 
+	FRHICommandList& RHICmdList, 
 	const FSceneView& View,
 	ContextType Context,
 	const FMeshBatch& Mesh,
@@ -1208,7 +1208,7 @@ void DrawMeshElements(FShadowDepthDrawingPolicy<bReflectiveShadowmap>& SharedDra
 #endif
 
 	//@todo-rco: RHIPacketList
-	FRHICommandList* RHICmdList = nullptr;
+	FRHICommandList& RHICmdList = FRHICommandList::GetNullRef();
 
 	// Render only those batch elements that match the current LOD
 	uint64 BatchElementMask = Mesh->Elements.Num() == 1 ? 1 : View.StaticMeshBatchVisibility[Mesh->Id];
@@ -1228,7 +1228,7 @@ void DrawMeshElements(FShadowDepthDrawingPolicy<bReflectiveShadowmap>& SharedDra
 }
 
 template <bool bReflectiveShadowmap>
-void DrawShadowMeshElements(FRHICommandList* RHICmdList, const FViewInfo& View, const FProjectedShadowInfo& ShadowInfo)
+void DrawShadowMeshElements(FRHICommandList& RHICmdList, const FViewInfo& View, const FProjectedShadowInfo& ShadowInfo)
 {
 	const FShadowStaticMeshElement& FirstShadowMesh = ShadowInfo.SubjectMeshElements[0];
 	const FMaterial* FirstMaterialResource = FirstShadowMesh.MaterialResource;
@@ -1267,7 +1267,7 @@ void DrawShadowMeshElements(FRHICommandList* RHICmdList, const FViewInfo& View,
 	}
 }
 
-void FProjectedShadowInfo::RenderDepth(FRHICommandList* RHICmdList, FDeferredShadingSceneRenderer* SceneRenderer)
+void FProjectedShadowInfo::RenderDepth(FRHICommandList& RHICmdList, FDeferredShadingSceneRenderer* SceneRenderer)
 {
 #if WANTS_DRAW_MESH_EVENTS
 	FString EventName;
@@ -1371,7 +1371,7 @@ void FProjectedShadowInfo::RenderDepth(FRHICommandList* RHICmdList, FDeferredSha
 	if (IsWholeSceneDirectionalShadow())
 	{
 		//@todo-rco: RHIPacketList
-		FRHICommandList* RHICmdList = nullptr;
+		FRHICommandList& RHICmdList = FRHICommandList::GetNullRef();
 
 		SCOPE_CYCLE_COUNTER(STAT_WholeSceneStaticDrawListShadowDepthsTime);
 
@@ -1468,7 +1468,7 @@ void StencilingGeometry::DrawCone()
 static FGlobalBoundShaderState MaskBoundShaderState[2];
 
 template <uint32 Quality>
-static void SetShadowProjectionShaderTemplNew(FRHICommandList* RHICmdList, int32 ViewIndex, const FViewInfo& View, const FProjectedShadowInfo* ShadowInfo)
+static void SetShadowProjectionShaderTemplNew(FRHICommandList& RHICmdList, int32 ViewIndex, const FViewInfo& View, const FProjectedShadowInfo* ShadowInfo)
 {
 	if (ShadowInfo->bTranslucentShadow)
 	{
@@ -1532,9 +1532,9 @@ static void SetShadowProjectionShaderTemplNew(FRHICommandList* RHICmdList, int32
 	}
 }
 
-void FProjectedShadowInfo::RenderProjection(FRHICommandList* RHICmdList, int32 ViewIndex, const FViewInfo* View) const
+void FProjectedShadowInfo::RenderProjection(FRHICommandList& RHICmdList, int32 ViewIndex, const FViewInfo* View) const
 {
-	check(!RHICmdList);
+	RHICmdList.CheckIsNull();
 
 #if WANTS_DRAW_MESH_EVENTS
 	FString EventName;
@@ -1943,7 +1943,7 @@ void FProjectedShadowInfo::RenderProjection(FRHICommandList* RHICmdList, int32 V
 
 
 template <uint32 Quality>
-static void SetPointLightShaderTempl(FRHICommandList* RHICmdList, int32 ViewIndex, const FViewInfo& View, const FProjectedShadowInfo* ShadowInfo)
+static void SetPointLightShaderTempl(FRHICommandList& RHICmdList, int32 ViewIndex, const FViewInfo& View, const FProjectedShadowInfo* ShadowInfo)
 {
 	TShaderMapRef<FShadowProjectionVS> VertexShader(GetGlobalShaderMap());
 	TShaderMapRef<TOnePassPointShadowProjectionPS<Quality> > PixelShader(GetGlobalShaderMap());
@@ -1956,11 +1956,11 @@ static void SetPointLightShaderTempl(FRHICommandList* RHICmdList, int32 ViewInde
 }
 
 /** Render one pass point light shadow projections. */
-void FProjectedShadowInfo::RenderOnePassPointLightProjection(FRHICommandList* RHICmdList, int32 ViewIndex, const FViewInfo& View) const
+void FProjectedShadowInfo::RenderOnePassPointLightProjection(FRHICommandList& RHICmdList, int32 ViewIndex, const FViewInfo& View) const
 {
 	SCOPE_CYCLE_COUNTER(STAT_RenderWholeSceneShadowProjectionsTime);
 
-	check(!RHICmdList);
+	RHICmdList.CheckIsNull();
 
 	checkSlow(bOnePassPointLightShadow);
 	
@@ -2336,7 +2336,7 @@ bool FDeferredShadingSceneRenderer::CheckForProjectedShadows( const FLightSceneI
 }
 
 /** Renders one pass point light shadows. */
-bool FDeferredShadingSceneRenderer::RenderOnePassPointLightShadows(FRHICommandList* RHICmdList, const FLightSceneInfo* LightSceneInfo, bool bRenderedTranslucentObjectShadows, bool& bInjectedTranslucentVolume)
+bool FDeferredShadingSceneRenderer::RenderOnePassPointLightShadows(FRHICommandList& RHICmdList, const FLightSceneInfo* LightSceneInfo, bool bRenderedTranslucentObjectShadows, bool& bInjectedTranslucentVolume)
 {
 	bool bDirty = false;
 	FVisibleLightInfo& VisibleLightInfo = VisibleLightInfos[LightSceneInfo->Id];
@@ -2406,12 +2406,12 @@ bool FDeferredShadingSceneRenderer::RenderOnePassPointLightShadows(FRHICommandLi
 
 /** Renders the projections of the given Shadows to the appropriate color render target. */
 void FDeferredShadingSceneRenderer::RenderProjections(
-	FRHICommandList* RHICmdList, 
+	FRHICommandList& RHICmdList, 
 	const FLightSceneInfo* LightSceneInfo, 
 	const TArray<FProjectedShadowInfo*,SceneRenderingAllocator>& Shadows
 	)
 {
-	check(!RHICmdList);
+	RHICmdList.CheckIsNull();
 
 	// Normal shadows render to light attenuation
 	GSceneRenderTargets.BeginRenderingLightAttenuation();
@@ -2492,7 +2492,7 @@ struct FCompareFProjectedShadowInfoBySplitIndex
 	}
 };
 
-bool FDeferredShadingSceneRenderer::RenderTranslucentProjectedShadows(FRHICommandList* RHICmdList, const FLightSceneInfo* LightSceneInfo )
+bool FDeferredShadingSceneRenderer::RenderTranslucentProjectedShadows(FRHICommandList& RHICmdList, const FLightSceneInfo* LightSceneInfo )
 {
 	SCOPE_CYCLE_COUNTER(STAT_ProjectedShadowDrawTime);
 
@@ -2651,7 +2651,7 @@ bool FDeferredShadingSceneRenderer::RenderTranslucentProjectedShadows(FRHIComman
 * @param LightSceneInfo Represents the current light
 * @return true if anything got rendered
 */
-bool FDeferredShadingSceneRenderer::RenderReflectiveShadowMaps(FRHICommandList* RHICmdList, const FLightSceneInfo* LightSceneInfo )
+bool FDeferredShadingSceneRenderer::RenderReflectiveShadowMaps(FRHICommandList& RHICmdList, const FLightSceneInfo* LightSceneInfo )
 {
 	SCOPE_CYCLE_COUNTER(STAT_ReflectiveShadowMapDrawTime);
 
@@ -2823,7 +2823,7 @@ bool FDeferredShadingSceneRenderer::RenderReflectiveShadowMaps(FRHICommandList*
  * @param LightSceneInfo Represents the current light
  * @return true if anything got rendered
  */
-bool FDeferredShadingSceneRenderer::RenderProjectedShadows(FRHICommandList* RHICmdList, const FLightSceneInfo* LightSceneInfo, bool bRenderedTranslucentObjectShadows, bool& bInjectedTranslucentVolume )
+bool FDeferredShadingSceneRenderer::RenderProjectedShadows(FRHICommandList& RHICmdList, const FLightSceneInfo* LightSceneInfo, bool bRenderedTranslucentObjectShadows, bool& bInjectedTranslucentVolume )
 {
 	SCOPE_CYCLE_COUNTER(STAT_ProjectedShadowDrawTime);
 
@@ -3001,9 +3001,9 @@ bool FDeferredShadingSceneRenderer::RenderProjectedShadows(FRHICommandList* RHIC
 }
 
 /** Renders preshadow depths for any preshadows whose depths aren't cached yet, and renders the projections of preshadows with opaque relevance. */
-bool FDeferredShadingSceneRenderer::RenderCachedPreshadows(FRHICommandList* RHICmdList, const FLightSceneInfo* LightSceneInfo)
+bool FDeferredShadingSceneRenderer::RenderCachedPreshadows(FRHICommandList& RHICmdList, const FLightSceneInfo* LightSceneInfo)
 {
-	check(!RHICmdList);
+	RHICmdList.CheckIsNull();
 
 	bool bAttenuationBufferDirty = false;
 
diff --git a/Engine/Source/Runtime/Renderer/Private/ShadowRendering.h b/Engine/Source/Runtime/Renderer/Private/ShadowRendering.h
--- a/Engine/Source/Runtime/Renderer/Private/ShadowRendering.h
+++ b/Engine/Source/Runtime/Renderer/Private/ShadowRendering.h
@@ -333,7 +333,7 @@ public:
 			&& bUsePositionOnlyVS == Other.bUsePositionOnlyVS
 			&& bPreShadow == Other.bPreShadow;
 	}
-	void DrawShared(FRHICommandList* RHICmdList, const FSceneView* View,FBoundShaderStateRHIParamRef BoundShaderState) const;
+	void DrawShared(FRHICommandList& RHICmdList, const FSceneView* View,FBoundShaderStateRHIParamRef BoundShaderState) const;
 
 	/** 
 	 * Create bound shader state using the vertex decl from the mesh draw policy
@@ -343,7 +343,7 @@ public:
 	FBoundShaderStateRHIRef CreateBoundShaderState(ERHIFeatureLevel::Type InFeatureLevel);
 
 	void SetMeshRenderState(
-		FRHICommandList* RHICmdList, 
+		FRHICommandList& RHICmdList, 
 		const FSceneView& View,
 		const FPrimitiveSceneProxy* PrimitiveSceneProxy,
 		const FMeshBatch& Mesh,
@@ -413,7 +413,7 @@ public:
 	static void AddStaticMesh(FScene* Scene,FStaticMesh* StaticMesh);
 
 	static bool DrawDynamicMesh(
-		FRHICommandList* RHICmdList, 
+		FRHICommandList& RHICmdList, 
 		const FSceneView& View,
 		ContextType Context,
 		const FMeshBatch& Mesh,
@@ -662,7 +662,7 @@ public:
 	/**
 	 * Renders the shadow subject depth.
 	 */
-	void RenderDepth(FRHICommandList* RHICmdList, class FDeferredShadingSceneRenderer* SceneRenderer);
+	void RenderDepth(FRHICommandList& RHICmdList, class FDeferredShadingSceneRenderer* SceneRenderer);
 
 	void ClearDepth(class FDeferredShadingSceneRenderer* SceneRenderer);
 
@@ -672,10 +672,10 @@ public:
 	/**
 	 * Projects the shadow onto the scene for a particular view.
 	 */
-	void RenderProjection(FRHICommandList* RHICmdList, int32 ViewIndex, const class FViewInfo* View) const;
+	void RenderProjection(FRHICommandList& RHICmdList, int32 ViewIndex, const class FViewInfo* View) const;
 
 	/** Render one pass point light shadow projections. */
-	void RenderOnePassPointLightProjection(FRHICommandList* RHICmdList, int32 ViewIndex, const FViewInfo& View) const;
+	void RenderOnePassPointLightProjection(FRHICommandList& RHICmdList, int32 ViewIndex, const FViewInfo& View) const;
 
 	/**
 	 * Renders the projected shadow's frustum wireframe with the given FPrimitiveDrawInterface.
@@ -743,7 +743,7 @@ private:
 
 	void GetShadowTypeNameForDrawEvent(FString& TypeName) const;
 
-	template <bool bReflectiveShadowmap> friend void DrawShadowMeshElements(FRHICommandList* RHICmdList, const FViewInfo& View, const FProjectedShadowInfo& ShadowInfo);
+	template <bool bReflectiveShadowmap> friend void DrawShadowMeshElements(FRHICommandList& RHICmdList, const FViewInfo& View, const FProjectedShadowInfo& ShadowInfo);
 };
 
 /** Shader parameters for rendering the depth of a mesh for shadowing. */
@@ -759,7 +759,7 @@ public:
 	}
 
 	template<typename ShaderRHIParamRef>
-	void Set(FRHICommandList* RHICmdList, ShaderRHIParamRef ShaderRHI, const FSceneView& View, const FProjectedShadowInfo* ShadowInfo, const FMaterialRenderProxy* MaterialRenderProxy)
+	void Set(FRHICommandList& RHICmdList, ShaderRHIParamRef ShaderRHI, const FSceneView& View, const FProjectedShadowInfo* ShadowInfo, const FMaterialRenderProxy* MaterialRenderProxy)
 	{
 		SetShaderValue(
 			RHICmdList, 
@@ -775,13 +775,13 @@ public:
 	}
 
 	/** Set the vertex shader parameter values. */
-	void SetVertexShader(FRHICommandList* RHICmdList, FShader* VertexShader, const FSceneView& View, const FProjectedShadowInfo* ShadowInfo, const FMaterialRenderProxy* MaterialRenderProxy)
+	void SetVertexShader(FRHICommandList& RHICmdList, FShader* VertexShader, const FSceneView& View, const FProjectedShadowInfo* ShadowInfo, const FMaterialRenderProxy* MaterialRenderProxy)
 	{
 		Set(RHICmdList, VertexShader->GetVertexShader(), View, ShadowInfo, MaterialRenderProxy);
 	}
 
 	/** Set the domain shader parameter values. */
-	void SetDomainShader(FRHICommandList* RHICmdList, FShader* DomainShader, const FSceneView& View, const FProjectedShadowInfo* ShadowInfo, const FMaterialRenderProxy* MaterialRenderProxy)
+	void SetDomainShader(FRHICommandList& RHICmdList, FShader* DomainShader, const FSceneView& View, const FProjectedShadowInfo* ShadowInfo, const FMaterialRenderProxy* MaterialRenderProxy)
 	{
 		Set(RHICmdList, DomainShader->GetDomainShader(), View, ShadowInfo, MaterialRenderProxy);
 	}
@@ -815,13 +815,13 @@ public:
 		StencilPreViewTranslation.Bind(ParameterMap, TEXT("StencilingPreViewTranslation"));
 	}
 
-	void Set(FRHICommandList* RHICmdList, FShader* Shader, const FVector4& InStencilingGeometryPosAndScale) const
+	void Set(FRHICommandList& RHICmdList, FShader* Shader, const FVector4& InStencilingGeometryPosAndScale) const
 	{
 		SetShaderValue(RHICmdList, Shader->GetVertexShader(), StencilGeometryPosAndScale, InStencilingGeometryPosAndScale);
 		SetShaderValue(RHICmdList, Shader->GetVertexShader(), StencilConeParameters, FVector4(0.0f, 0.0f, 0.0f, 0.0f));
 	}
 
-	void Set(FRHICommandList* RHICmdList, FShader* Shader, const FSceneView& View, const FLightSceneInfo* LightSceneInfo) const
+	void Set(FRHICommandList& RHICmdList, FShader* Shader, const FSceneView& View, const FLightSceneInfo* LightSceneInfo) const
 	{
 		FVector4 GeometryPosAndScale;
 		if(LightSceneInfo->Proxy->GetLightType() == LightType_Point)
@@ -885,7 +885,7 @@ public:
 		OutEnvironment.SetDefine(TEXT("USE_TRANSFORM"), (uint32)1);
 	}
 
-	void SetParameters(FRHICommandList* RHICmdList, const FSceneView& View, const FProjectedShadowInfo* ShadowInfo, int32 ShadowSplitIndex = INDEX_NONE);
+	void SetParameters(FRHICommandList& RHICmdList, const FSceneView& View, const FProjectedShadowInfo* ShadowInfo, int32 ShadowSplitIndex = INDEX_NONE);
 
 	// Begin FShader interface
 	virtual bool Serialize(FArchive& Ar) OVERRIDE
@@ -925,7 +925,7 @@ public:
 		return true;
 	}
 
-	inline void SetParameters(FRHICommandList* RHICmdList, const FSceneView& View)
+	inline void SetParameters(FRHICommandList& RHICmdList, const FSceneView& View)
 	{
 		FGlobalShader::SetParameters(RHICmdList, GetVertexShader(),View); 
 	}
@@ -958,7 +958,7 @@ public:
 	 * @param ShadowInfo - projected shadow info for a single light
 	 */
 	virtual void SetParameters(
-		FRHICommandList* RHICmdList, 
+		FRHICommandList& RHICmdList, 
 		int32 ViewIndex,
 		const FSceneView& View,
 		const FProjectedShadowInfo* ShadowInfo
@@ -992,7 +992,7 @@ public:
 		InvFadePlaneLength.Bind(ParameterMap,TEXT("InvFadePlaneLength"));
 	}
 
-	void Set(FRHICommandList* RHICmdList, FShader* Shader, const FSceneView& View, const FProjectedShadowInfo* ShadowInfo)
+	void Set(FRHICommandList& RHICmdList, FShader* Shader, const FSceneView& View, const FProjectedShadowInfo* ShadowInfo)
 	{
 		const FPixelShaderRHIParamRef ShaderRHI = Shader->GetPixelShader();
 
@@ -1027,8 +1027,7 @@ public:
 
 		if (ShadowDepthSampler.IsBound())
 		{
-			FRHICommandList::SetShaderSampler(
-				RHICmdList,
+			RHICmdList.SetShaderSampler(
 				ShaderRHI, 
 				ShadowDepthSampler.GetBaseIndex(), 
 				DepthSamplerState
@@ -1127,7 +1126,7 @@ public:
 	 * @param ShadowInfo - projected shadow info for a single light
 	 */
 	virtual void SetParameters(
-		FRHICommandList* RHICmdList, 
+		FRHICommandList& RHICmdList, 
 		int32 ViewIndex,
 		const FSceneView& View,
 		const FProjectedShadowInfo* ShadowInfo
@@ -1173,7 +1172,7 @@ public:
 		TranslucencyShadowTransmission1Sampler.Bind(ParameterMap,TEXT("TranslucencyShadowTransmission1Sampler"));
 	}
 
-	void Set(FRHICommandList* RHICmdList, FShader* Shader) const
+	void Set(FRHICommandList& RHICmdList, FShader* Shader) const
 	{
 		SetTextureParameter(
 			RHICmdList, 
@@ -1239,7 +1238,7 @@ public:
 	}
 
 	virtual void SetParameters(
-		FRHICommandList* RHICmdList, 
+		FRHICommandList& RHICmdList, 
 		int32 ViewIndex,
 		const FSceneView& View,
 		const FProjectedShadowInfo* ShadowInfo) OVERRIDE
@@ -1279,7 +1278,7 @@ public:
 	}
 
 	template<typename ShaderRHIParamRef>
-	void Set(FRHICommandList* RHICmdList, const ShaderRHIParamRef ShaderRHI, const FProjectedShadowInfo* ShadowInfo) const
+	void Set(FRHICommandList& RHICmdList, const ShaderRHIParamRef ShaderRHI, const FProjectedShadowInfo* ShadowInfo) const
 	{
 		SetTextureParameter(
 			RHICmdList, 
@@ -1363,7 +1362,7 @@ public:
 	}
 
 	void SetParameters(
-		FRHICommandList* RHICmdList, 
+		FRHICommandList& RHICmdList, 
 		int32 ViewIndex,
 		const FSceneView& View,
 		const FProjectedShadowInfo* ShadowInfo
diff --git a/Engine/Source/Runtime/Renderer/Private/StaticMeshDrawList.h b/Engine/Source/Runtime/Renderer/Private/StaticMeshDrawList.h
--- a/Engine/Source/Runtime/Renderer/Private/StaticMeshDrawList.h
+++ b/Engine/Source/Runtime/Renderer/Private/StaticMeshDrawList.h
@@ -228,7 +228,7 @@ private:
 	* @param DrawingPolicyLink - the drawing policy link
 	* @param bDrawnShared - determines whether to draw shared 
 	*/
-	void DrawElement(FRHICommandList* RHICmdList, const FViewInfo& View, const FElement& Element, uint64 BatchElementMask, FDrawingPolicyLink* DrawingPolicyLink, bool &bDrawnShared);
+	void DrawElement(FRHICommandList& RHICmdList, const FViewInfo& View, const FElement& Element, uint64 BatchElementMask, FDrawingPolicyLink* DrawingPolicyLink, bool &bDrawnShared);
 
 public:
 
@@ -261,7 +261,7 @@ public:
 	 * @param BatchVisibilityArray - An array of batch element visibility bitmasks.
 	 * @return True if any static meshes were drawn.
 	 */
-	bool DrawVisible(FRHICommandList* RHICmdList, const FViewInfo& View, const TBitArray<SceneRenderingBitArrayAllocator>& StaticMeshVisibilityMap, const TArray<uint64,SceneRenderingAllocator>& BatchVisibilityArray);
+	bool DrawVisible(FRHICommandList& RHICmdList, const FViewInfo& View, const TBitArray<SceneRenderingBitArrayAllocator>& StaticMeshVisibilityMap, const TArray<uint64,SceneRenderingAllocator>& BatchVisibilityArray);
 
 	/**
 	 * Draws only the static meshes which are in the visibility map, sorted front-to-back.
@@ -271,7 +271,7 @@ public:
 	 * @param MaxToDraw - The maximum number of meshes to be drawn.
 	 * @return The number of static meshes drawn.
 	 */
-	int32 DrawVisibleFrontToBack(FRHICommandList* RHICmdList, const FViewInfo& View, const TBitArray<SceneRenderingBitArrayAllocator>& StaticMeshVisibilityMap, const TArray<uint64,SceneRenderingAllocator>& BatchVisibilityArray, int32 MaxToDraw);
+	int32 DrawVisibleFrontToBack(FRHICommandList& RHICmdList, const FViewInfo& View, const TBitArray<SceneRenderingBitArrayAllocator>& StaticMeshVisibilityMap, const TArray<uint64,SceneRenderingAllocator>& BatchVisibilityArray, int32 MaxToDraw);
 
 	/** Sorts OrderedDrawingPolicies front to back. */
 	void SortFrontToBack(FVector ViewPosition);
diff --git a/Engine/Source/Runtime/Renderer/Private/StaticMeshDrawList.inl b/Engine/Source/Runtime/Renderer/Private/StaticMeshDrawList.inl
--- a/Engine/Source/Runtime/Renderer/Private/StaticMeshDrawList.inl
+++ b/Engine/Source/Runtime/Renderer/Private/StaticMeshDrawList.inl
@@ -58,7 +58,7 @@ void TStaticMeshDrawList<DrawingPolicyType>::FElementHandle::Remove()
 
 template<typename DrawingPolicyType>
 void TStaticMeshDrawList<DrawingPolicyType>::DrawElement(
-	FRHICommandList* RHICmdList,
+	FRHICommandList& RHICmdList,
 	const FViewInfo& View,
 	const FElement& Element,
 	uint64 BatchElementMask,
@@ -238,7 +238,7 @@ bool TStaticMeshDrawList<DrawingPolicyType>::DrawVisible(
 
 template<typename DrawingPolicyType>
 bool TStaticMeshDrawList<DrawingPolicyType>::DrawVisible(
-	FRHICommandList* RHICmdList,
+	FRHICommandList& RHICmdList,
 	const FViewInfo& View,
 	const TBitArray<SceneRenderingBitArrayAllocator>& StaticMeshVisibilityMap,
 	const TArray<uint64,SceneRenderingAllocator>& BatchVisibilityArray
@@ -271,7 +271,7 @@ bool TStaticMeshDrawList<DrawingPolicyType>::DrawVisible(
 
 template<typename DrawingPolicyType>
 int32 TStaticMeshDrawList<DrawingPolicyType>::DrawVisibleFrontToBack(
-	FRHICommandList* RHICmdList,
+	FRHICommandList& RHICmdList,
 	const FViewInfo& View,
 	const TBitArray<SceneRenderingBitArrayAllocator>& StaticMeshVisibilityMap,
 	const TArray<uint64,SceneRenderingAllocator>& BatchVisibilityArray,
diff --git a/Engine/Source/Runtime/Renderer/Private/TiledDeferredLightRendering.cpp b/Engine/Source/Runtime/Renderer/Private/TiledDeferredLightRendering.cpp
--- a/Engine/Source/Runtime/Renderer/Private/TiledDeferredLightRendering.cpp
+++ b/Engine/Source/Runtime/Renderer/Private/TiledDeferredLightRendering.cpp
@@ -106,7 +106,7 @@ public:
 	}
 
 	void SetParameters(
-		FRHICommandList* RHICmdList, 
+		FRHICommandList& RHICmdList, 
 		const FSceneView& View, 
 		int32 ViewIndex,
 		int32 NumViews,
@@ -211,7 +211,7 @@ public:
 		SetShaderValue(RHICmdList, ShaderRHI, NumLights, NumThisPass);
 	}
 
-	void UnsetParameters(FRHICommandList* RHICmdList)
+	void UnsetParameters(FRHICommandList& RHICmdList)
 	{
 		OutTexture.UnsetUAV(RHICmdList, GetComputeShader());
 	}
@@ -272,7 +272,7 @@ bool FDeferredShadingSceneRenderer::ShouldUseTiledDeferred(int32 NumUnshadowedLi
 
 template <bool bVisualizeLightCulling>
 static void SetShaderTemplTiledLighting(
-	FRHICommandList* RHICmdList, 
+	FRHICommandList& RHICmdList, 
 	const FSceneView& View, 
 	int32 ViewIndex,
 	int32 NumViews,
@@ -297,9 +297,9 @@ static void SetShaderTemplTiledLighting(
 }
 
 
-void FDeferredShadingSceneRenderer::RenderTiledDeferredLighting(FRHICommandList* RHICmdList, const TArray<FSortedLightSceneInfo, SceneRenderingAllocator>& SortedLights, int32 NumUnshadowedLights, const FSimpleLightArray& SimpleLights)
+void FDeferredShadingSceneRenderer::RenderTiledDeferredLighting(FRHICommandList& RHICmdList, const TArray<FSortedLightSceneInfo, SceneRenderingAllocator>& SortedLights, int32 NumUnshadowedLights, const FSimpleLightArray& SimpleLights)
 {
-	check(!RHICmdList);
+	RHICmdList.CheckIsNull();
 
 	check(GUseTiledDeferredShading);
 	check(SortedLights.Num() >= NumUnshadowedLights);
diff --git a/Engine/Source/Runtime/Renderer/Private/TranslucentLighting.cpp b/Engine/Source/Runtime/Renderer/Private/TranslucentLighting.cpp
--- a/Engine/Source/Runtime/Renderer/Private/TranslucentLighting.cpp
+++ b/Engine/Source/Runtime/Renderer/Private/TranslucentLighting.cpp
@@ -182,7 +182,7 @@ public:
 	}
 
 	void SetParameters(
-		FRHICommandList* RHICmdList, 
+		FRHICommandList& RHICmdList, 
 		const FMaterialRenderProxy* MaterialRenderProxy,
 		const FSceneView& View,
 		const FProjectedShadowInfo* ShadowInfo
@@ -192,7 +192,7 @@ public:
 		ShadowParameters.SetVertexShader(RHICmdList, this, View, ShadowInfo, MaterialRenderProxy);
 	}
 
-	void SetMesh(FRHICommandList* RHICmdList, const FVertexFactory* VertexFactory,const FSceneView& View,const FPrimitiveSceneProxy* Proxy,const FMeshBatchElement& BatchElement)
+	void SetMesh(FRHICommandList& RHICmdList, const FVertexFactory* VertexFactory,const FSceneView& View,const FPrimitiveSceneProxy* Proxy,const FMeshBatchElement& BatchElement)
 	{
 		FMeshMaterialShader::SetMesh(RHICmdList, GetVertexShader(),VertexFactory,View,Proxy,BatchElement);
 	}
@@ -254,7 +254,7 @@ public:
 	FTranslucencyShadowDepthPS() {}
 
 	void SetParameters(
-		FRHICommandList* RHICmdList, 
+		FRHICommandList& RHICmdList, 
 		const FMaterialRenderProxy* MaterialRenderProxy,
 		const FSceneView& View,
 		const FProjectedShadowInfo* ShadowInfo
@@ -276,7 +276,7 @@ public:
 		TranslucencyProjectionParameters.Set(RHICmdList, this);
 	}
 
-	void SetMesh(FRHICommandList* RHICmdList, const FVertexFactory* VertexFactory,const FSceneView& View,const FPrimitiveSceneProxy* Proxy,const FMeshBatchElement& BatchElement)
+	void SetMesh(FRHICommandList& RHICmdList, const FVertexFactory* VertexFactory,const FSceneView& View,const FPrimitiveSceneProxy* Proxy,const FMeshBatchElement& BatchElement)
 	{
 		FMeshMaterialShader::SetMesh(RHICmdList, GetPixelShader(),VertexFactory,View,Proxy,BatchElement);
 	}
@@ -348,9 +348,9 @@ public:
 		}
 	}
 
-	void DrawShared(FRHICommandList* RHICmdList, const FSceneView* View,FBoundShaderStateRHIParamRef BoundShaderState) const
+	void DrawShared(FRHICommandList& RHICmdList, const FSceneView* View,FBoundShaderStateRHIParamRef BoundShaderState) const
 	{
-		check(!RHICmdList);
+		RHICmdList.CheckIsNull();
 
 		// Set the shared mesh resources.
 		FMeshDrawingPolicy::DrawShared(RHICmdList, View);
@@ -379,7 +379,7 @@ public:
 	}
 
 	void SetMeshRenderState(
-		FRHICommandList* RHICmdList, 
+		FRHICommandList& RHICmdList, 
 		const FSceneView& View,
 		const FPrimitiveSceneProxy* PrimitiveSceneProxy,
 		const FMeshBatch& Mesh,
@@ -419,7 +419,7 @@ public:
 	};
 
 	static bool DrawDynamicMesh(
-		FRHICommandList* RHICmdList, 
+		FRHICommandList& RHICmdList, 
 		const FSceneView& View,
 		ContextType DrawingContext,
 		const FMeshBatch& Mesh,
@@ -456,7 +456,7 @@ public:
 	}
 
 	static bool DrawStaticMesh(
-		FRHICommandList* RHICmdList, 
+		FRHICommandList& RHICmdList, 
 		const FViewInfo& View,
 		ContextType DrawingContext,
 		const FStaticMesh& StaticMesh,
@@ -491,7 +491,7 @@ public:
 void FProjectedShadowInfo::RenderTranslucencyDepths(FDeferredShadingSceneRenderer* SceneRenderer)
 {
 	//@todo-rco: RHIPacketList
-	FRHICommandList* RHICmdList = nullptr;
+	FRHICommandList& RHICmdList = FRHICommandList::GetNullRef();
 
 	checkSlow(!bWholeSceneShadow);
 	SCOPE_CYCLE_COUNTER(STAT_RenderWholeSceneShadowDepthsTime);
@@ -644,7 +644,7 @@ public:
 	}
 	FFilterTranslucentVolumePS() {}
 
-	void SetParameters(FRHICommandList* RHICmdList, const FViewInfo& View, int32 VolumeCascadeIndex)
+	void SetParameters(FRHICommandList& RHICmdList, const FViewInfo& View, int32 VolumeCascadeIndex)
 	{
 		const FPixelShaderRHIParamRef ShaderRHI = GetPixelShader();
 
@@ -703,7 +703,7 @@ public:
 
 	template<typename ShaderRHIParamRef>
 	void Set(
-		FRHICommandList* RHICmdList, 
+		FRHICommandList& RHICmdList, 
 		const ShaderRHIParamRef ShaderRHI, 
 		FShader* Shader, 
 		const FViewInfo& View, 
@@ -767,7 +767,7 @@ public:
 	}
 	FTranslucentObjectShadowingPS() {}
 
-	void SetParameters(FRHICommandList* RHICmdList, const FViewInfo& View, const FLightSceneInfo* LightSceneInfo, const FProjectedShadowInfo* ShadowMap, int32 VolumeCascadeIndexValue)
+	void SetParameters(FRHICommandList& RHICmdList, const FViewInfo& View, const FLightSceneInfo* LightSceneInfo, const FProjectedShadowInfo* ShadowMap, int32 VolumeCascadeIndexValue)
 	{
 		FGlobalShader::SetParameters(RHICmdList, GetPixelShader(), View);
 		TranslucencyProjectionParameters.Set(RHICmdList, this);
@@ -837,7 +837,7 @@ public:
 	// @param InnerSplitIndex which CSM shadow map level, INDEX_NONE if no directional light
 	// @param VolumeCascadeIndexValue which volume we render to
 	void SetParameters(
-		FRHICommandList* RHICmdList, 
+		FRHICommandList& RHICmdList, 
 		const FViewInfo& View, 
 		const FLightSceneInfo* LightSceneInfo, 
 		const FMaterialRenderProxy* MaterialProxy, 
@@ -1028,9 +1028,9 @@ void RasterizeToVolumeTexture(FVolumeBounds VolumeBounds)
 
 /** Helper function that clears the given volume texture render targets. */
 template<int32 NumRenderTargets>
-void ClearVolumeTextures(FRHICommandList* RHICmdList, const FTextureRHIParamRef* RenderTargets, const FLinearColor* ClearColors)
+void ClearVolumeTextures(FRHICommandList& RHICmdList, const FTextureRHIParamRef* RenderTargets, const FLinearColor* ClearColors)
 {
-	check(!RHICmdList);
+	RHICmdList.CheckIsNull();
 
 	RHISetRenderTargets(NumRenderTargets, RenderTargets, FTextureRHIRef(), 0, NULL);
 
@@ -1071,7 +1071,7 @@ void ClearVolumeTextures(FRHICommandList* RHICmdList, const FTextureRHIParamRef*
 	}
 }
 
-void FDeferredShadingSceneRenderer::ClearTranslucentVolumeLighting(FRHICommandList* RHICmdList)
+void FDeferredShadingSceneRenderer::ClearTranslucentVolumeLighting(FRHICommandList& RHICmdList)
 {
 	if (GUseTranslucentLightingVolumes && GSupportsVolumeTextureRendering)
 	{
@@ -1128,7 +1128,7 @@ public:
 		return bShaderHasOutdatedParameters;
 	}
 
-	void SetParameters(FRHICommandList* RHICmdList, const FViewInfo& View, const FFinalPostProcessSettings::FCubemapEntry& CubemapEntry)
+	void SetParameters(FRHICommandList& RHICmdList, const FViewInfo& View, const FFinalPostProcessSettings::FCubemapEntry& CubemapEntry)
 	{
 		const FPixelShaderRHIParamRef ShaderRHI = GetPixelShader();
 
@@ -1140,7 +1140,7 @@ public:
 
 IMPLEMENT_SHADER_TYPE(,FInjectAmbientCubemapPS,TEXT("TranslucentLightingShaders"),TEXT("InjectAmbientCubemapMainPS"),SF_Pixel);
 
-void FDeferredShadingSceneRenderer::InjectAmbientCubemapTranslucentVolumeLighting(FRHICommandList* RHICmdList)
+void FDeferredShadingSceneRenderer::InjectAmbientCubemapTranslucentVolumeLighting(FRHICommandList& RHICmdList)
 {
 	//@todo - support multiple views
 	const FViewInfo& View = Views[0];
@@ -1214,7 +1214,7 @@ public:
 	}
 	FCompositeGIForTranslucencyPS() {}
 
-	void SetParameters(FRHICommandList* RHICmdList, const FViewInfo& View, const FSceneViewState* ViewState, int32 VolumeCascadeIndexValue)
+	void SetParameters(FRHICommandList& RHICmdList, const FViewInfo& View, const FSceneViewState* ViewState, int32 VolumeCascadeIndexValue)
 	{
 		const FPixelShaderRHIParamRef ShaderRHI = GetPixelShader();
 
@@ -1293,7 +1293,7 @@ private:
 
 IMPLEMENT_SHADER_TYPE(,FCompositeGIForTranslucencyPS,TEXT("TranslucentLightingShaders"),TEXT("CompositeGIMainPS"),SF_Pixel);
 
-void FDeferredShadingSceneRenderer::CompositeIndirectTranslucentVolumeLighting(FRHICommandList* RHICmdList)
+void FDeferredShadingSceneRenderer::CompositeIndirectTranslucentVolumeLighting(FRHICommandList& RHICmdList)
 {
 	bool bAnyViewAllowsIndirectLightingCache = false;
 
@@ -1369,7 +1369,7 @@ void FDeferredShadingSceneRenderer::ClearTranslucentVolumePerObjectShadowing()
 		ClearColors[1] = FLinearColor(1, 1, 1, 1);
 
 		//@todo-rco: RHIPacketList
-		FRHICommandList* RHICmdList = nullptr;
+		FRHICommandList& RHICmdList = FRHICommandList::GetNullRef();
 		ClearVolumeTextures<ARRAY_COUNT(RenderTargets)>(RHICmdList, RenderTargets, ClearColors);
 	}
 }
@@ -1413,7 +1413,7 @@ void FDeferredShadingSceneRenderer::AccumulateTranslucentVolumeObjectShadowing(c
 	}
 
 	//@todo-rco: RHIPacketList
-	FRHICommandList* RHICmdList = nullptr;
+	FRHICommandList& RHICmdList = FRHICommandList::GetNullRef();
 	if (GUseTranslucentLightingVolumes && GSupportsVolumeTextureRendering)
 	{
 		SCOPED_DRAW_EVENT(AccumulateTranslucentVolumeShadowing, DEC_SCENE_ITEMS);
@@ -1484,7 +1484,7 @@ void SetInjectionShader(
 	bool bInverseSquared)
 {
 	//@todo-rco: RHIPacketList
-	FRHICommandList* RHICmdList = nullptr;
+	FRHICommandList& RHICmdList = FRHICommandList::GetNullRef();
 
 	check(ShadowMap || !bShadowed);
 
@@ -1633,7 +1633,7 @@ static void InjectTranslucentLightArray(const FViewInfo& View, const TArray<FTra
 	INC_DWORD_STAT_BY(STAT_NumLightsInjectedIntoTranslucency, LightInjectionData.Num());
 
 	//@todo-rco: RHIPacketList
-	FRHICommandList* RHICmdList = nullptr;
+	FRHICommandList& RHICmdList = FRHICommandList::GetNullRef();
 
 	// Inject into each volume cascade
 	// Operate on one cascade at a time to reduce render target switches
@@ -1775,7 +1775,7 @@ public:
 	}
 	FSimpleLightTranslucentLightingInjectPS() {}
 
-	void SetParameters(FRHICommandList* RHICmdList, const FViewInfo& View, const FSimpleLightEntry& SimpleLight, const FSimpleLightPerViewEntry& SimpleLightPerViewData, int32 VolumeCascadeIndexValue)
+	void SetParameters(FRHICommandList& RHICmdList, const FViewInfo& View, const FSimpleLightEntry& SimpleLight, const FSimpleLightPerViewEntry& SimpleLightPerViewData, int32 VolumeCascadeIndexValue)
 	{
 		FGlobalShader::SetParameters(RHICmdList, GetPixelShader(), View);
 
@@ -1807,7 +1807,7 @@ FGlobalBoundShaderState InjectSimpleLightBoundShaderState;
 void FDeferredShadingSceneRenderer::InjectSimpleTranslucentVolumeLightingArray(const FSimpleLightArray& SimpleLights)
 {
 	//@todo-rco: RHIPacketList
-	FRHICommandList* RHICmdList = nullptr;
+	FRHICommandList& RHICmdList = FRHICommandList::GetNullRef();
 
 	SCOPE_CYCLE_COUNTER(STAT_TranslucentInjectTime);
 
@@ -1895,7 +1895,7 @@ void FDeferredShadingSceneRenderer::FilterTranslucentVolumeLighting()
 		if (GUseTranslucencyVolumeBlur)
 		{
 			//@todo-rco: RHIPacketList
-			FRHICommandList* RHICmdList = nullptr;
+			FRHICommandList& RHICmdList = FRHICommandList::GetNullRef();
 
 			SCOPED_DRAW_EVENT(FilterTranslucentVolume, DEC_SCENE_ITEMS);
 
diff --git a/Engine/Source/Runtime/Renderer/Private/TranslucentRendering.cpp b/Engine/Source/Runtime/Renderer/Private/TranslucentRendering.cpp
--- a/Engine/Source/Runtime/Renderer/Private/TranslucentRendering.cpp
+++ b/Engine/Source/Runtime/Renderer/Private/TranslucentRendering.cpp
@@ -146,7 +146,7 @@ public:
 	}
 	FCopySceneColorPS() {}
 
-	void SetParameters(FRHICommandList* RHICmdList, const FViewInfo& View)
+	void SetParameters(FRHICommandList& RHICmdList, const FViewInfo& View)
 	{
 		SceneTextureParameters.Set(RHICmdList, GetPixelShader(), View);
 	}
@@ -183,7 +183,7 @@ void FDeferredShadingSceneRenderer::CopySceneColor(const FViewInfo& View, const
 	SetGlobalBoundShaderState(CopySceneColorBoundShaderState, GFilterVertexDeclaration.VertexDeclarationRHI, *ScreenVertexShader, *PixelShader);
 
 	/// ?
-	PixelShader->SetParameters(nullptr, View);
+	PixelShader->SetParameters(FRHICommandList::GetNullRef(), View);
 
 	DrawRectangle( 
 		0, 0, 
@@ -248,7 +248,7 @@ public:
 	/** Draws the translucent mesh with a specific light-map type, and fog volume type */
 	template<typename LightMapPolicyType>
 	void Process(
-		FRHICommandList* RHICmdList, 
+		FRHICommandList& RHICmdList, 
 		const FProcessBasePassMeshParameters& Parameters,
 		const LightMapPolicyType& LightMapPolicy,
 		const typename LightMapPolicyType::ElementDataType& LightMapElementData
@@ -325,7 +325,7 @@ bool FTranslucencyDrawingPolicyFactory::DrawMesh(
 	const EMaterialShadingModel ShadingModel = Material->GetShadingModel();
 
 	//@todo-rco: RHIPacketList
-	FRHICommandList* RHICmdList = nullptr;
+	FRHICommandList& RHICmdList = FRHICommandList::GetNullRef();
 
 	// Only render translucent materials.
 	if(IsTranslucentBlendMode(BlendMode))
@@ -391,7 +391,7 @@ bool FTranslucencyDrawingPolicyFactory::DrawMesh(
  * @return true if the mesh rendered
  */
 bool FTranslucencyDrawingPolicyFactory::DrawDynamicMesh(
-	FRHICommandList* RHICmdList, 
+	FRHICommandList& RHICmdList, 
 	const FViewInfo& View,
 	ContextType DrawingContext,
 	const FMeshBatch& Mesh,
@@ -417,7 +417,7 @@ bool FTranslucencyDrawingPolicyFactory::DrawDynamicMesh(
  * @return true if the mesh rendered
  */
 bool FTranslucencyDrawingPolicyFactory::DrawStaticMesh(
-	FRHICommandList* RHICmdList, 
+	FRHICommandList& RHICmdList, 
 	const FViewInfo& View,
 	ContextType DrawingContext,
 	const FStaticMesh& StaticMesh,
@@ -524,7 +524,7 @@ void FTranslucentPrimSet::RenderPrimitive(
 		}
 
 		//@todo-rco: RHIPacketList
-		FRHICommandList* RHICmdList = nullptr;
+		FRHICommandList& RHICmdList = FRHICommandList::GetNullRef();
 
 		// Render static scene prim
 		if( ViewRelevance.bStaticRelevance )
diff --git a/Engine/Source/Runtime/Renderer/Private/TranslucentRendering.h b/Engine/Source/Runtime/Renderer/Private/TranslucentRendering.h
--- a/Engine/Source/Runtime/Renderer/Private/TranslucentRendering.h
+++ b/Engine/Source/Runtime/Renderer/Private/TranslucentRendering.h
@@ -28,7 +28,7 @@ public:
 	* @return true if the mesh rendered
 	*/
 	static bool DrawDynamicMesh(
-		FRHICommandList* RHICmdList, 
+		FRHICommandList& RHICmdList, 
 		const FViewInfo& View,
 		ContextType DrawingContext,
 		const FMeshBatch& Mesh,
@@ -43,7 +43,7 @@ public:
 	* @return true if the mesh rendered
 	*/
 	static bool DrawStaticMesh(
-		FRHICommandList* RHICmdList, 
+		FRHICommandList& RHICmdList, 
 		const FViewInfo& View,
 		ContextType DrawingContext,
 		const FStaticMesh& StaticMesh,
@@ -95,7 +95,7 @@ public:
 	* @return true if the mesh rendered
 	*/
 	static bool DrawDynamicMesh(
-		FRHICommandList* RHICmdList, 
+		FRHICommandList& RHICmdList, 
 		const FViewInfo& View,
 		ContextType DrawingContext,
 		const FMeshBatch& Mesh,
@@ -110,7 +110,7 @@ public:
 	* @return true if the mesh rendered
 	*/
 	static bool DrawStaticMesh(
-		FRHICommandList* RHICmdList, 
+		FRHICommandList& RHICmdList, 
 		const FViewInfo& View,
 		ContextType DrawingContext,
 		const FStaticMesh& StaticMesh,
@@ -175,7 +175,7 @@ public:
 
 	FWriteToSliceVS() {}
 
-	void SetParameters(FRHICommandList* RHICmdList, const FVolumeBounds& VolumeBounds, uint32 VolumeResolution)
+	void SetParameters(FRHICommandList& RHICmdList, const FVolumeBounds& VolumeBounds, uint32 VolumeResolution)
 	{
 		const float InvVolumeResolution = 1.0f / VolumeResolution;
 		SetShaderValue(RHICmdList, GetVertexShader(), UVScaleBias, FVector4(
@@ -214,7 +214,7 @@ public:
 	}
 	FWriteToSliceGS() {}
 
-	void SetParameters(FRHICommandList* RHICmdList, const FVolumeBounds& VolumeBounds)
+	void SetParameters(FRHICommandList& RHICmdList, const FVolumeBounds& VolumeBounds)
 	{
 		SetShaderValue(RHICmdList, GetGeometryShader(), MinZ, VolumeBounds.MinZ);
 	}
diff --git a/Engine/Source/Runtime/Renderer/Private/VelocityRendering.cpp b/Engine/Source/Runtime/Renderer/Private/VelocityRendering.cpp
--- a/Engine/Source/Runtime/Renderer/Private/VelocityRendering.cpp
+++ b/Engine/Source/Runtime/Renderer/Private/VelocityRendering.cpp
@@ -15,12 +15,12 @@ class FVelocityVS : public FMeshMaterialShader
 
 public:
 
-	void SetParameters(FRHICommandList* RHICmdList, const FVertexFactory* VertexFactory,const FMaterialRenderProxy* MaterialRenderProxy,const FViewInfo& View)
+	void SetParameters(FRHICommandList& RHICmdList, const FVertexFactory* VertexFactory,const FMaterialRenderProxy* MaterialRenderProxy,const FViewInfo& View)
 	{
 		FMeshMaterialShader::SetParameters(RHICmdList, GetVertexShader(), MaterialRenderProxy, *MaterialRenderProxy->GetMaterial(View.GetFeatureLevel()), View, ESceneRenderTargetsMode::DontSet);
 	}
 
-	void SetMesh(FRHICommandList* RHICmdList, const FVertexFactory* VertexFactory, const FMeshBatch& Mesh, int32 BatchElementIndex, const FViewInfo& View, const FPrimitiveSceneProxy* Proxy, const FMatrix& InPreviousLocalToWorld)
+	void SetMesh(FRHICommandList& RHICmdList, const FVertexFactory* VertexFactory, const FMeshBatch& Mesh, int32 BatchElementIndex, const FViewInfo& View, const FPrimitiveSceneProxy* Proxy, const FMatrix& InPreviousLocalToWorld)
 	{
 		FMeshMaterialShader::SetMesh(RHICmdList, GetVertexShader(), VertexFactory, View, Proxy, Mesh.Elements[BatchElementIndex]);
 
@@ -91,7 +91,7 @@ class FVelocityDS : public FBaseDS
 
 public:
 
-	void SetParameters(FRHICommandList* RHICmdList, const FMaterialRenderProxy* MaterialRenderProxy,const FViewInfo& View)
+	void SetParameters(FRHICommandList& RHICmdList, const FMaterialRenderProxy* MaterialRenderProxy,const FViewInfo& View)
 	{
 		FMeshMaterialShader::SetParameters(RHICmdList, GetDomainShader(), MaterialRenderProxy, *MaterialRenderProxy->GetMaterial(View.GetFeatureLevel()), View, ESceneRenderTargetsMode::DontSet);
 	}
@@ -144,12 +144,12 @@ public:
 	{
 	}
 
-	void SetParameters(FRHICommandList* RHICmdList, const FVertexFactory* VertexFactory,const FMaterialRenderProxy* MaterialRenderProxy,const FViewInfo& View)
+	void SetParameters(FRHICommandList& RHICmdList, const FVertexFactory* VertexFactory,const FMaterialRenderProxy* MaterialRenderProxy,const FViewInfo& View)
 	{
 		FMeshMaterialShader::SetParameters(RHICmdList, GetPixelShader(), MaterialRenderProxy, *MaterialRenderProxy->GetMaterial(View.GetFeatureLevel()), View, ESceneRenderTargetsMode::DontSet);
 	}
 
-	void SetMesh(FRHICommandList* RHICmdList, const FVertexFactory* VertexFactory, const FMeshBatch& Mesh,int32 BatchElementIndex,const FViewInfo& View, const FPrimitiveSceneProxy* Proxy, bool bBackFace)
+	void SetMesh(FRHICommandList& RHICmdList, const FVertexFactory* VertexFactory, const FMeshBatch& Mesh,int32 BatchElementIndex,const FViewInfo& View, const FPrimitiveSceneProxy* Proxy, bool bBackFace)
 	{
 		const FPixelShaderRHIParamRef ShaderRHI = GetPixelShader();
 		
@@ -206,13 +206,13 @@ bool FVelocityDrawingPolicy::SupportsVelocity() const
 	return (VertexShader && PixelShader) ? VertexShader->SupportsVelocity() : false;
 }
 
-void FVelocityDrawingPolicy::DrawShared(FRHICommandList* RHICmdList, const FSceneView* SceneView, FBoundShaderStateRHIRef ShaderState) const
+void FVelocityDrawingPolicy::DrawShared(FRHICommandList& RHICmdList, const FSceneView* SceneView, FBoundShaderStateRHIRef ShaderState) const
 {
 	// NOTE: Assuming this cast is always safe!
 	FViewInfo* View = (FViewInfo*)SceneView;
 
 	// Set the depth-only shader parameters for the material.
-	FRHICommandList::SetBoundShaderState(RHICmdList, ShaderState);
+	RHICmdList.SetBoundShaderState(ShaderState);
 
 	VertexShader->SetParameters(RHICmdList, VertexFactory, MaterialRenderProxy, *View);
 	PixelShader->SetParameters(RHICmdList, VertexFactory, MaterialRenderProxy, *View);
@@ -228,7 +228,7 @@ void FVelocityDrawingPolicy::DrawShared(FRHICommandList* RHICmdList, const FScen
 }
 
 void FVelocityDrawingPolicy::SetMeshRenderState(
-	FRHICommandList* RHICmdList, 
+	FRHICommandList& RHICmdList, 
 	const FSceneView& SceneView,
 	const FPrimitiveSceneProxy* PrimitiveSceneProxy,
 	const FMeshBatch& Mesh,
@@ -372,7 +372,7 @@ void FVelocityDrawingPolicyFactory::AddStaticMesh(FScene* Scene, FStaticMesh* St
 }
 
 bool FVelocityDrawingPolicyFactory::DrawDynamicMesh(
-	FRHICommandList* RHICmdList, 
+	FRHICommandList& RHICmdList, 
 	const FViewInfo& View,
 	ContextType DrawingContext,
 	const FMeshBatch& Mesh,
@@ -473,7 +473,7 @@ void FDeferredShadingSceneRenderer::RenderVelocities(const FViewInfo& View, TRef
 	GPrevPerBoneMotionBlur.LockData();
 
 	//@todo-rco: RHIPacketList
-	FRHICommandList* RHICmdList = nullptr;
+	FRHICommandList& RHICmdList = FRHICommandList::GetNullRef();
 
 	const uint32 MinX = View.ViewRect.Min.X * VelocityBufferSize.X / BufferSize.X;
 	const uint32 MinY = View.ViewRect.Min.Y * VelocityBufferSize.Y / BufferSize.Y;
diff --git a/Engine/Source/Runtime/Renderer/Private/VelocityRendering.h b/Engine/Source/Runtime/Renderer/Private/VelocityRendering.h
--- a/Engine/Source/Runtime/Renderer/Private/VelocityRendering.h
+++ b/Engine/Source/Runtime/Renderer/Private/VelocityRendering.h
@@ -27,9 +27,9 @@ public:
 			&& VertexShader == Other.VertexShader
 			&& PixelShader == Other.PixelShader;
 	}
-	void DrawShared(FRHICommandList* RHICmdList, const FSceneView* View, FBoundShaderStateRHIRef ShaderState) const;
+	void DrawShared(FRHICommandList& RHICmdList, const FSceneView* View, FBoundShaderStateRHIRef ShaderState) const;
 	void SetMeshRenderState(
-		FRHICommandList* RHICmdList, 
+		FRHICommandList& RHICmdList, 
 		const FSceneView& View,
 		const FPrimitiveSceneProxy* PrimitiveSceneProxy,
 		const FMeshBatch& Mesh,
@@ -62,7 +62,7 @@ class FVelocityDrawingPolicyFactory : public FDepthDrawingPolicyFactory
 public:
 	static void AddStaticMesh(FScene* Scene, FStaticMesh* StaticMesh, ContextType = ContextType(DDM_AllOccluders));
 	static bool DrawDynamicMesh(	
-		FRHICommandList* RHICmdList, 
+		FRHICommandList& RHICmdList, 
 		const FViewInfo& View,
 		ContextType DrawingContext,
 		const FMeshBatch& Mesh,
diff --git a/Engine/Source/Runtime/ShaderCore/Private/Shader.cpp b/Engine/Source/Runtime/ShaderCore/Private/Shader.cpp
--- a/Engine/Source/Runtime/ShaderCore/Private/Shader.cpp
+++ b/Engine/Source/Runtime/ShaderCore/Private/Shader.cpp
@@ -11,6 +11,7 @@
 #include "DerivedDataCacheInterface.h"
 #include "ModuleManager.h"
 #include "TargetPlatform.h"
+#include "RHICommandList.h"
 
 
 DEFINE_LOG_CATEGORY(LogShaders);
@@ -1033,13 +1034,13 @@ FShaderType* FindShaderTypeByName(const TCHAR* ShaderTypeName)
 
 
 void DispatchComputeShader(
-	FRHICommandList* RHICmdList, 
+	FRHICommandList& RHICmdList, 
 	FShader* Shader,
 	uint32 ThreadGroupCountX,
 	uint32 ThreadGroupCountY,
 	uint32 ThreadGroupCountZ)
 {
-	check(!RHICmdList);
+	RHICmdList.CheckIsNull();
 
 	Shader->VerifyBoundUniformBufferParameters();
 	RHIDispatchComputeShader(ThreadGroupCountX, ThreadGroupCountY, ThreadGroupCountZ);
@@ -1047,12 +1048,12 @@ void DispatchComputeShader(
 
 
 void DispatchIndirectComputeShader(
-	FRHICommandList* RHICmdList, 
+	FRHICommandList& RHICmdList, 
 	FShader* Shader,
 	FVertexBufferRHIParamRef ArgumentBuffer,
 	uint32 ArgumentOffset)
 {
-	check(!RHICmdList);
+	RHICmdList.CheckIsNull();
 
 	Shader->VerifyBoundUniformBufferParameters();
 	RHIDispatchIndirectComputeShader(ArgumentBuffer, ArgumentOffset);
diff --git a/Engine/Source/Runtime/ShaderCore/Private/VertexFactory.cpp b/Engine/Source/Runtime/ShaderCore/Private/VertexFactory.cpp
--- a/Engine/Source/Runtime/ShaderCore/Private/VertexFactory.cpp
+++ b/Engine/Source/Runtime/ShaderCore/Private/VertexFactory.cpp
@@ -167,21 +167,21 @@ FVertexFactoryType* FindVertexFactoryType(FName TypeName)
 	return NULL;
 }
 
-void FVertexFactory::Set(FRHICommandList* RHICmdList) const
+void FVertexFactory::Set(FRHICommandList& RHICmdList) const
 {
 	check(IsInitialized());
 	for(int32 StreamIndex = 0;StreamIndex < Streams.Num();StreamIndex++)
 	{
 		const FVertexStream& Stream = Streams[StreamIndex];
 		checkf(Stream.VertexBuffer->IsInitialized(), TEXT("Vertex buffer was not initialized! Stream %u, Stride %u, Name %s"), StreamIndex, Stream.Stride, *Stream.VertexBuffer->GetFriendlyName());
-		FRHICommandList::SetStreamSource(RHICmdList, StreamIndex, Stream.VertexBuffer->VertexBufferRHI, Stream.Stride, Stream.Offset);
+		RHICmdList.SetStreamSource( StreamIndex, Stream.VertexBuffer->VertexBufferRHI, Stream.Stride, Stream.Offset);
 	}
 }
 
 void FVertexFactory::SetPositionStream() const
 {
 	//@todo-rco: RHIPacketList
-	FRHICommandList* RHICmdList = nullptr;
+	FRHICommandList& RHICmdList = FRHICommandList::GetNullRef();
 
 	check(IsInitialized());
 	// Set the predefined vertex streams.
@@ -189,7 +189,7 @@ void FVertexFactory::SetPositionStream() const
 	{
 		const FVertexStream& Stream = PositionStream[StreamIndex];
 		check(Stream.VertexBuffer->IsInitialized());
-		FRHICommandList::SetStreamSource(RHICmdList, StreamIndex, Stream.VertexBuffer->VertexBufferRHI, Stream.Stride, Stream.Offset);
+		RHICmdList.SetStreamSource( StreamIndex, Stream.VertexBuffer->VertexBufferRHI, Stream.Stride, Stream.Offset);
 	}
 }
 
diff --git a/Engine/Source/Runtime/ShaderCore/Public/Shader.h b/Engine/Source/Runtime/ShaderCore/Public/Shader.h
--- a/Engine/Source/Runtime/ShaderCore/Public/Shader.h
+++ b/Engine/Source/Runtime/ShaderCore/Public/Shader.h
@@ -1320,7 +1320,7 @@ extern SHADERCORE_API FShaderType* FindShaderTypeByName(const TCHAR* ShaderTypeN
 
 /** Helper function to dispatch a compute shader while checking that parameters have been set correctly. */
 extern SHADERCORE_API void DispatchComputeShader(
-	FRHICommandList* RHICmdList,
+	FRHICommandList& RHICmdList,
 	FShader* Shader,
 	uint32 ThreadGroupCountX,
 	uint32 ThreadGroupCountY,
@@ -1328,7 +1328,7 @@ extern SHADERCORE_API void DispatchComputeShader(
 
 /** Helper function to dispatch a compute shader indirectly while checking that parameters have been set correctly. */
 extern SHADERCORE_API void DispatchIndirectComputeShader(
-	FRHICommandList* RHICmdList, 
+	FRHICommandList& RHICmdList, 
 	FShader* Shader,
 	FVertexBufferRHIParamRef ArgumentBuffer,
 	uint32 ArgumentOffset);
diff --git a/Engine/Source/Runtime/ShaderCore/Public/ShaderParameterUtils.h b/Engine/Source/Runtime/ShaderCore/Public/ShaderParameterUtils.h
--- a/Engine/Source/Runtime/ShaderCore/Public/ShaderParameterUtils.h
+++ b/Engine/Source/Runtime/ShaderCore/Public/ShaderParameterUtils.h
@@ -18,7 +18,7 @@
  */
 template<typename ShaderRHIParamRef, class ParameterType>
 void SetShaderValue(
-	FRHICommandList* RHICmdList, 
+	FRHICommandList& RHICmdList, 
 	ShaderRHIParamRef Shader,
 	const FShaderParameter& Parameter,
 	const ParameterType& Value,
@@ -35,8 +35,7 @@ void SetShaderValue(
 
 	if(NumBytesToSet > 0)
 	{
-		FRHICommandList::SetShaderParameter(
-			RHICmdList,
+		RHICmdList.SetShaderParameter(
 			Shader,
 			Parameter.GetBufferIndex(),
 			Parameter.GetBaseIndex() + ElementIndex * AlignedTypeSize,
@@ -50,7 +49,7 @@ void SetShaderValue(
 /** Specialization of the above for C++ bool type. */
 template<typename ShaderRHIParamRef>
 void SetShaderValue(
-	FRHICommandList* RHICmdList, 
+	FRHICommandList& RHICmdList, 
 	ShaderRHIParamRef Shader,
 	const FShaderParameter& Parameter,
 	const bool& Value,
@@ -69,7 +68,7 @@ void SetShaderValue(
  */
 template<typename ShaderRHIParamRef,class ParameterType>
 void SetShaderValueArray(
-	FRHICommandList* RHICmdList,
+	FRHICommandList& RHICmdList,
 	ShaderRHIParamRef Shader,
 	const FShaderParameter& Parameter,
 	const ParameterType* Values,
@@ -85,8 +84,7 @@ void SetShaderValueArray(
 
 	if(NumBytesToSet > 0)
 	{
-		FRHICommandList::SetShaderParameter(
-			RHICmdList,
+		RHICmdList.SetShaderParameter(
 			Shader,
 			Parameter.GetBufferIndex(),
 			Parameter.GetBaseIndex() + BaseElementIndex * AlignedTypeSize,
@@ -100,7 +98,7 @@ void SetShaderValueArray(
 /** Specialization of the above for C++ bool type. */
 template<typename ShaderRHIParamRef>
 void SetShaderValueArray(
-	FRHICommandList* RHICmdList, 
+	FRHICommandList& RHICmdList, 
 	ShaderRHIParamRef Shader,
 	const FShaderParameter& Parameter,
 	const bool* Values,
@@ -115,7 +113,7 @@ void SetShaderValueArray(
  * Sets the value of a pixel shader bool parameter.
  */
 inline void SetPixelShaderBool(
-	FRHICommandList* RHICmdList, 
+	FRHICommandList& RHICmdList, 
 	FPixelShaderRHIParamRef PixelShader,
 	const FShaderParameter& Parameter,
 	bool Value
@@ -128,8 +126,7 @@ inline void SetPixelShaderBool(
 	{
 		// Convert to uint32 before passing to RHI
 		uint32 BoolValue = Value;
-		FRHICommandList::SetShaderParameter(
-			RHICmdList,
+		RHICmdList.SetShaderParameter(
 			PixelShader,
 			Parameter.GetBufferIndex(),
 			Parameter.GetBaseIndex(),
@@ -145,7 +142,7 @@ inline void SetPixelShaderBool(
  */
 template<typename ShaderTypeRHIParamRef>
 FORCEINLINE void SetTextureParameter(
-	FRHICommandList* RHICmdList, 
+	FRHICommandList& RHICmdList, 
 	ShaderTypeRHIParamRef Shader,
 	const FShaderResourceParameter& TextureParameter,
 	const FShaderResourceParameter& SamplerParameter,
@@ -162,7 +159,7 @@ FORCEINLINE void SetTextureParameter(
 
 		if (ElementIndex < TextureParameter.GetNumResources())
 		{
-			FRHICommandList::SetShaderTexture(RHICmdList, Shader, TextureParameter.GetBaseIndex() + ElementIndex, Texture->TextureRHI);
+			RHICmdList.SetShaderTexture( Shader, TextureParameter.GetBaseIndex() + ElementIndex, Texture->TextureRHI);
 		}
 	}
 	
@@ -176,7 +173,7 @@ FORCEINLINE void SetTextureParameter(
 	{
 		if (ElementIndex < SamplerParameter.GetNumResources())
 		{
-			FRHICommandList::SetShaderSampler(RHICmdList, Shader, SamplerParameter.GetBaseIndex() + ElementIndex, Texture->SamplerStateRHI);
+			RHICmdList.SetShaderSampler( Shader, SamplerParameter.GetBaseIndex() + ElementIndex, Texture->SamplerStateRHI);
 		}
 	}
 }
@@ -186,7 +183,7 @@ FORCEINLINE void SetTextureParameter(
  */
 template<typename ShaderTypeRHIParamRef>
 FORCEINLINE void SetTextureParameter(
-	FRHICommandList* RHICmdList, 
+	FRHICommandList& RHICmdList, 
 	ShaderTypeRHIParamRef Shader,
 	const FShaderResourceParameter& TextureParameter,
 	const FShaderResourceParameter& SamplerParameter,
@@ -202,7 +199,7 @@ FORCEINLINE void SetTextureParameter(
 	{
 		if (ElementIndex < TextureParameter.GetNumResources())
 		{
-			FRHICommandList::SetShaderTexture(RHICmdList, Shader, TextureParameter.GetBaseIndex() + ElementIndex, TextureRHI);
+			RHICmdList.SetShaderTexture( Shader, TextureParameter.GetBaseIndex() + ElementIndex, TextureRHI);
 		}
 	}
 	// @todo ue4 samplerstate Should we maybe pass in two separate values? SamplerElement and TextureElement? Or never allow an array of samplers? Unsure best
@@ -215,7 +212,7 @@ FORCEINLINE void SetTextureParameter(
 	{
 		if (ElementIndex < SamplerParameter.GetNumResources())
 		{
-			FRHICommandList::SetShaderSampler(RHICmdList, Shader, SamplerParameter.GetBaseIndex() + ElementIndex, SamplerStateRHI);
+			RHICmdList.SetShaderSampler( Shader, SamplerParameter.GetBaseIndex() + ElementIndex, SamplerStateRHI);
 		}
 	}
 }
@@ -226,7 +223,7 @@ FORCEINLINE void SetTextureParameter(
  */
 template<typename ShaderTypeRHIParamRef>
 FORCEINLINE void SetTextureParameter(
-	FRHICommandList* RHICmdList, 
+	FRHICommandList& RHICmdList, 
 	ShaderTypeRHIParamRef Shader,
 	const FShaderResourceParameter& Parameter,
 	FTextureRHIParamRef NewTextureRHI
@@ -234,8 +231,7 @@ FORCEINLINE void SetTextureParameter(
 {
 	if(Parameter.IsBound())
 	{
-		FRHICommandList::SetShaderTexture(
-			RHICmdList,
+		RHICmdList.SetShaderTexture(
 			Shader,
 			Parameter.GetBaseIndex(),
 			NewTextureRHI
@@ -248,7 +244,7 @@ FORCEINLINE void SetTextureParameter(
  */
 template<typename ShaderTypeRHIParamRef>
 FORCEINLINE void SetSamplerParameter(
-	FRHICommandList* RHICmdList, 
+	FRHICommandList& RHICmdList, 
 	ShaderTypeRHIParamRef Shader,
 	const FShaderResourceParameter& Parameter,
 	FSamplerStateRHIParamRef SamplerStateRHI
@@ -256,8 +252,7 @@ FORCEINLINE void SetSamplerParameter(
 {
 	if(Parameter.IsBound())
 	{
-		FRHICommandList::SetShaderSampler(
-			RHICmdList,
+		RHICmdList.SetShaderSampler(
 			Shader,
 			Parameter.GetBaseIndex(),
 			SamplerStateRHI
@@ -271,7 +266,7 @@ FORCEINLINE void SetSamplerParameter(
  */
 template<typename ShaderTypeRHIParamRef>
 FORCEINLINE void SetSRVParameter(
-	FRHICommandList* RHICmdList, 
+	FRHICommandList& RHICmdList, 
 	ShaderTypeRHIParamRef Shader,
 	const FShaderResourceParameter& Parameter,
 	FShaderResourceViewRHIParamRef NewShaderResourceViewRHI
@@ -279,8 +274,7 @@ FORCEINLINE void SetSRVParameter(
 {
 	if(Parameter.IsBound())
 	{
-		FRHICommandList::SetShaderResourceViewParameter(
-			RHICmdList,
+		RHICmdList.SetShaderResourceViewParameter(
 			Shader,
 			Parameter.GetBaseIndex(),
 			NewShaderResourceViewRHI
@@ -292,7 +286,7 @@ FORCEINLINE void SetSRVParameter(
  * Sets the value of a unordered access view parameter
  */
 FORCEINLINE void SetUAVParameter(
-	FRHICommandList* RHICmdList, 
+	FRHICommandList& RHICmdList, 
 	FComputeShaderRHIParamRef ComputeShader,
 	const FShaderResourceParameter& Parameter,
 	FUnorderedAccessViewRHIParamRef NewUnorderedAccessViewRHI
@@ -300,8 +294,7 @@ FORCEINLINE void SetUAVParameter(
 {
 	if(Parameter.IsBound())
 	{
-		FRHICommandList::SetUAVParameter(
-			RHICmdList,
+		RHICmdList.SetUAVParameter(
 			ComputeShader,
 			Parameter.GetBaseIndex(),
 			NewUnorderedAccessViewRHI
@@ -310,35 +303,35 @@ FORCEINLINE void SetUAVParameter(
 }
 
 template<typename TShaderRHIRef>
-inline void SetUAVParameterIfCS(FRHICommandList* RHICmdList, TShaderRHIRef,const FShaderResourceParameter& UAVParameter,FUnorderedAccessViewRHIParamRef UAV)
+inline void SetUAVParameterIfCS(FRHICommandList& RHICmdList, TShaderRHIRef,const FShaderResourceParameter& UAVParameter,FUnorderedAccessViewRHIParamRef UAV)
 {
 }
 
-template<> inline void SetUAVParameterIfCS<FComputeShaderRHIParamRef>(FRHICommandList* RHICmdList, FComputeShaderRHIParamRef ComputeShader,const FShaderResourceParameter& UAVParameter,FUnorderedAccessViewRHIParamRef UAV)
+template<> inline void SetUAVParameterIfCS<FComputeShaderRHIParamRef>(FRHICommandList& RHICmdList, FComputeShaderRHIParamRef ComputeShader,const FShaderResourceParameter& UAVParameter,FUnorderedAccessViewRHIParamRef UAV)
 {
 	SetUAVParameter(RHICmdList, ComputeShader,UAVParameter,UAV);
 }
 
-template<> inline void SetUAVParameterIfCS<FComputeShaderRHIRef>(FRHICommandList* RHICmdList, FComputeShaderRHIRef ComputeShader,const FShaderResourceParameter& UAVParameter,FUnorderedAccessViewRHIParamRef UAV)
+template<> inline void SetUAVParameterIfCS<FComputeShaderRHIRef>(FRHICommandList& RHICmdList, FComputeShaderRHIRef ComputeShader,const FShaderResourceParameter& UAVParameter,FUnorderedAccessViewRHIParamRef UAV)
 {
 	SetUAVParameter(RHICmdList, ComputeShader,UAVParameter,UAV);
 }
 
 template<typename TShaderRHIRef>
-inline void FRWShaderParameter::SetBuffer(FRHICommandList* RHICmdList, TShaderRHIRef Shader,const FRWBuffer& RWBuffer) const
+inline void FRWShaderParameter::SetBuffer(FRHICommandList& RHICmdList, TShaderRHIRef Shader,const FRWBuffer& RWBuffer) const
 {
 	SetSRVParameter(RHICmdList, Shader,SRVParameter,RWBuffer.SRV);
 	SetUAVParameterIfCS(RHICmdList, Shader,UAVParameter,RWBuffer.UAV);
 }
 
 template<typename TShaderRHIRef>
-inline void FRWShaderParameter::SetTexture(FRHICommandList* RHICmdList, TShaderRHIRef Shader,const FTextureRHIParamRef Texture,FUnorderedAccessViewRHIParamRef UAV) const
+inline void FRWShaderParameter::SetTexture(FRHICommandList& RHICmdList, TShaderRHIRef Shader,const FTextureRHIParamRef Texture,FUnorderedAccessViewRHIParamRef UAV) const
 {
 	SetTextureParameter(RHICmdList, Shader,SRVParameter,Texture);
 	SetUAVParameterIfCS(RHICmdList, Shader,UAVParameter,UAV);
 }
 
-inline void FRWShaderParameter::UnsetUAV(FRHICommandList* RHICmdList, FComputeShaderRHIParamRef ComputeShader) const
+inline void FRWShaderParameter::UnsetUAV(FRHICommandList& RHICmdList, FComputeShaderRHIParamRef ComputeShader) const
 {
 	SetUAVParameter(RHICmdList, ComputeShader,UAVParameter,FUnorderedAccessViewRHIRef());
 }
@@ -347,7 +340,7 @@ inline void FRWShaderParameter::UnsetUAV(FRHICommandList* RHICmdList, FComputeSh
 /** Sets the value of a shader uniform buffer parameter to a uniform buffer containing the struct. */
 template<typename TShaderRHIRef>
 inline void SetUniformBufferParameter(
-	FRHICommandList* RHICmdList,
+	FRHICommandList& RHICmdList,
 	TShaderRHIRef Shader,
 	const FShaderUniformBufferParameter& Parameter,
 	FUniformBufferRHIParamRef UniformBufferRHI
@@ -357,14 +350,14 @@ inline void SetUniformBufferParameter(
 	checkSlow(Parameter.IsInitialized());
 	if(Parameter.IsBound())
 	{
-		FRHICommandList::SetShaderUniformBuffer(RHICmdList, Shader, Parameter.GetBaseIndex(), UniformBufferRHI);
+		RHICmdList.SetShaderUniformBuffer( Shader, Parameter.GetBaseIndex(), UniformBufferRHI);
 	}
 }
 
 /** Sets the value of a shader uniform buffer parameter to a uniform buffer containing the struct. */
 template<typename TShaderRHIRef,typename TBufferStruct>
 inline void SetUniformBufferParameter(
-	FRHICommandList* RHICmdList,
+	FRHICommandList& RHICmdList,
 	TShaderRHIRef Shader,
 	const TShaderUniformBufferParameter<TBufferStruct>& Parameter,
 	const TUniformBufferRef<TBufferStruct>& UniformBufferRef
@@ -374,14 +367,14 @@ inline void SetUniformBufferParameter(
 	checkSlow(Parameter.IsInitialized());
 	if(Parameter.IsBound())
 	{
-		FRHICommandList::SetShaderUniformBuffer(RHICmdList, Shader, Parameter.GetBaseIndex(), UniformBufferRef);
+		RHICmdList.SetShaderUniformBuffer( Shader, Parameter.GetBaseIndex(), UniformBufferRef);
 	}
 }
 
 /** Sets the value of a shader uniform buffer parameter to a uniform buffer containing the struct. */
 template<typename TShaderRHIRef,typename TBufferStruct>
 inline void SetUniformBufferParameter(
-	FRHICommandList* RHICmdList,
+	FRHICommandList& RHICmdList,
 	TShaderRHIRef Shader,
 	const TShaderUniformBufferParameter<TBufferStruct>& Parameter,
 	const TUniformBuffer<TBufferStruct>& UniformBuffer
@@ -391,14 +384,14 @@ inline void SetUniformBufferParameter(
 	checkSlow(Parameter.IsInitialized());
 	if (Parameter.IsBound())
 	{
-		FRHICommandList::SetShaderUniformBuffer(RHICmdList, Shader, Parameter.GetBaseIndex(), UniformBuffer.GetUniformBufferRHI());
+		RHICmdList.SetShaderUniformBuffer( Shader, Parameter.GetBaseIndex(), UniformBuffer.GetUniformBufferRHI());
 	}
 }
 
 /** Sets the value of a shader uniform buffer parameter to a value of the struct. */
 template<typename TShaderRHIRef,typename TBufferStruct>
 inline void SetUniformBufferParameterImmediate(
-	FRHICommandList* RHICmdList,
+	FRHICommandList& RHICmdList,
 	TShaderRHIRef Shader,
 	const TShaderUniformBufferParameter<TBufferStruct>& Parameter,
 	const TBufferStruct& UniformBufferValue
@@ -408,8 +401,7 @@ inline void SetUniformBufferParameterImmediate(
 	checkSlow(Parameter.IsInitialized());
 	if(Parameter.IsBound())
 	{
-		FRHICommandList::SetShaderUniformBuffer(
-			RHICmdList,
+		RHICmdList.SetShaderUniformBuffer(
 			Shader,
 			Parameter.GetBaseIndex(),
 			RHICreateUniformBuffer(&UniformBufferValue,TBufferStruct::StaticStruct.GetLayout(),UniformBuffer_SingleDraw)
diff --git a/Engine/Source/Runtime/ShaderCore/Public/ShaderParameters.h b/Engine/Source/Runtime/ShaderCore/Public/ShaderParameters.h
--- a/Engine/Source/Runtime/ShaderCore/Public/ShaderParameters.h
+++ b/Engine/Source/Runtime/ShaderCore/Public/ShaderParameters.h
@@ -128,12 +128,12 @@ public:
 	}
 
 	template<typename TShaderRHIRef>
-	inline void SetBuffer(FRHICommandList* RHICmdList, TShaderRHIRef Shader,const FRWBuffer& RWBuffer) const;
+	inline void SetBuffer(FRHICommandList& RHICmdList, TShaderRHIRef Shader,const FRWBuffer& RWBuffer) const;
 
 	template<typename TShaderRHIRef>
-	inline void SetTexture(FRHICommandList* RHICmdList, TShaderRHIRef Shader,const FTextureRHIParamRef Texture,FUnorderedAccessViewRHIParamRef UAV) const;
+	inline void SetTexture(FRHICommandList& RHICmdList, TShaderRHIRef Shader,const FTextureRHIParamRef Texture,FUnorderedAccessViewRHIParamRef UAV) const;
 
-	inline void UnsetUAV(FRHICommandList* RHICmdList, FComputeShaderRHIParamRef ComputeShader) const;
+	inline void UnsetUAV(FRHICommandList& RHICmdList, FComputeShaderRHIParamRef ComputeShader) const;
 
 private:
 
diff --git a/Engine/Source/Runtime/ShaderCore/Public/VertexFactory.h b/Engine/Source/Runtime/ShaderCore/Public/VertexFactory.h
--- a/Engine/Source/Runtime/ShaderCore/Public/VertexFactory.h
+++ b/Engine/Source/Runtime/ShaderCore/Public/VertexFactory.h
@@ -72,7 +72,7 @@ public:
 	virtual ~FVertexFactoryShaderParameters() {}
 	virtual void Bind(const class FShaderParameterMap& ParameterMap) = 0;
 	virtual void Serialize(FArchive& Ar) = 0;
-	virtual void SetMesh(FRHICommandList* RHICmdList, FShader* VertexShader,const class FVertexFactory* VertexFactory,const class FSceneView& View,const struct FMeshBatchElement& BatchElement,uint32 DataFlags) const = 0;
+	virtual void SetMesh(FRHICommandList& RHICmdList, FShader* VertexShader,const class FVertexFactory* VertexFactory,const class FSceneView& View,const struct FMeshBatchElement& BatchElement,uint32 DataFlags) const = 0;
 	virtual uint32 GetSize() const { return sizeof(*this); }
 };
 
@@ -339,7 +339,7 @@ public:
 	/**
 	 * Activates the vertex factory.
 	 */
-	void Set(FRHICommandList* RHICmdList) const;
+	void Set(FRHICommandList& RHICmdList) const;
 
 	/**
 	* Sets the position stream as the current stream source.
@@ -472,7 +472,7 @@ public:
 		delete Parameters;
 	}
 
-	void SetMesh(FRHICommandList* RHICmdList, FShader* Shader,const FVertexFactory* VertexFactory,const class FSceneView& View,const FMeshBatchElement& BatchElement,uint32 DataFlags) const
+	void SetMesh(FRHICommandList& RHICmdList, FShader* Shader,const FVertexFactory* VertexFactory,const class FSceneView& View,const FMeshBatchElement& BatchElement,uint32 DataFlags) const
 	{
 		if(Parameters)
 		{
diff --git a/Engine/Source/Runtime/SlateRHIRenderer/Private/SlateShaders.cpp b/Engine/Source/Runtime/SlateRHIRenderer/Private/SlateShaders.cpp
--- a/Engine/Source/Runtime/SlateRHIRenderer/Private/SlateShaders.cpp
+++ b/Engine/Source/Runtime/SlateRHIRenderer/Private/SlateShaders.cpp
@@ -68,12 +68,12 @@ FSlateElementVS::FSlateElementVS( const ShaderMetaType::CompiledShaderInitialize
 	VertexShaderParams.Bind( Initializer.ParameterMap, TEXT("VertexShaderParams"));
 }
 
-void FSlateElementVS::SetViewProjection(FRHICommandList* RHICmdList, const FMatrix& InViewProjection )
+void FSlateElementVS::SetViewProjection(FRHICommandList& RHICmdList, const FMatrix& InViewProjection )
 {
 	SetShaderValue(RHICmdList, GetVertexShader(), ViewProjection, InViewProjection );
 }
 
-void FSlateElementVS::SetShaderParameters(FRHICommandList* RHICmdList, const FVector4& ShaderParams )
+void FSlateElementVS::SetShaderParameters(FRHICommandList& RHICmdList, const FVector4& ShaderParams )
 {
 	SetShaderValue(RHICmdList, GetVertexShader(), VertexShaderParams, ShaderParams );
 }
diff --git a/Engine/Source/Runtime/SlateRHIRenderer/Private/SlateShaders.h b/Engine/Source/Runtime/SlateRHIRenderer/Private/SlateShaders.h
--- a/Engine/Source/Runtime/SlateRHIRenderer/Private/SlateShaders.h
+++ b/Engine/Source/Runtime/SlateRHIRenderer/Private/SlateShaders.h
@@ -46,14 +46,14 @@ public:
 	 *
 	 * @param InViewProjection	The ViewProjection matrix to use when this shader is bound 
 	 */
-	void SetViewProjection(FRHICommandList* RHICmdList, const FMatrix& InViewProjection );
+	void SetViewProjection(FRHICommandList& RHICmdList, const FMatrix& InViewProjection );
 
 	/** 
 	 * Sets shader parameters for use in this shader
 	 *
 	 * @param ShaderParams	The shader params to be used
 	 */
-	void SetShaderParameters(FRHICommandList* RHICmdList, const FVector4& ShaderParams );
+	void SetShaderParameters(FRHICommandList& RHICmdList, const FVector4& ShaderParams );
 
 	/** Serializes the shader data */
 	virtual bool Serialize( FArchive& Ar );
@@ -100,7 +100,7 @@ public:
 	 * @param Texture	Texture resource to use when this pixel shader is bound
 	 * @param SamplerState	Sampler state to use when sampling this texture
 	 */
-	void SetTexture(FRHICommandList* RHICmdList, const FTextureRHIParamRef InTexture, const FSamplerStateRHIRef SamplerState )
+	void SetTexture(FRHICommandList& RHICmdList, const FTextureRHIParamRef InTexture, const FSamplerStateRHIRef SamplerState )
 	{
 		SetTextureParameter(RHICmdList, GetPixelShader(), TextureParameter, TextureParameterSampler, SamplerState, InTexture );
 	}
@@ -110,7 +110,7 @@ public:
 	 *
 	 * @param InDrawEffects		A mask of draw effects that the pixel shader should use.
 	 */
-	void SetDrawEffects(FRHICommandList* RHICmdList, ESlateDrawEffect::Type InDrawEffects )
+	void SetDrawEffects(FRHICommandList& RHICmdList, ESlateDrawEffect::Type InDrawEffects )
 	{
 		SetShaderValue(RHICmdList, GetPixelShader(), DrawEffectsParameter, InDrawEffects );
 	}
@@ -120,7 +120,7 @@ public:
 	 * 
 	 * @param InShaderParams Shader params to use
 	 */
-	void SetShaderParams(FRHICommandList* RHICmdList, const FVector4& InShaderParams )
+	void SetShaderParams(FRHICommandList& RHICmdList, const FVector4& InShaderParams )
 	{
 		SetShaderValue(RHICmdList, GetPixelShader(), ShaderParams, InShaderParams );
 	}
@@ -128,7 +128,7 @@ public:
 	/**
 	 * Set the viewport size.
 	 */
-	void SetViewportSize(FRHICommandList* RHICmdList, const FVector2D& InViewportSize)
+	void SetViewportSize(FRHICommandList& RHICmdList, const FVector2D& InViewportSize)
 	{
 		SetShaderValue(RHICmdList, GetPixelShader(), ViewportSize, InViewportSize);
 	}
@@ -138,7 +138,7 @@ public:
  	 *
 	 * @param DisplayGamma The display gamma to use
 	 */
-	void SetDisplayGamma(FRHICommandList* RHICmdList, float InDisplayGamma)
+	void SetDisplayGamma(FRHICommandList& RHICmdList, float InDisplayGamma)
 	{
 		SetShaderValue(RHICmdList, GetPixelShader(),DisplayGamma,InDisplayGamma);
 	}
diff --git a/Engine/Source/Runtime/Windows/D3D11RHI/Private/D3D11RenderTarget.cpp b/Engine/Source/Runtime/Windows/D3D11RHI/Private/D3D11RenderTarget.cpp
--- a/Engine/Source/Runtime/Windows/D3D11RHI/Private/D3D11RenderTarget.cpp
+++ b/Engine/Source/Runtime/Windows/D3D11RHI/Private/D3D11RenderTarget.cpp
@@ -29,7 +29,7 @@ namespace
 		}
 		FResolveDepthPS() {}
 
-		void SetParameters(FRHICommandList* RHICmdList, ID3D11DeviceContext* Direct3DDeviceContext,FParameter)
+		void SetParameters(FRHICommandList& RHICmdList, ID3D11DeviceContext* Direct3DDeviceContext,FParameter)
 		{
 		}
 
@@ -61,7 +61,7 @@ namespace
 		}
 		FResolveDepthNonMSPS() {}
 
-		void SetParameters(FRHICommandList* RHICmdList, ID3D11DeviceContext* Direct3DDeviceContext,FParameter)
+		void SetParameters(FRHICommandList& RHICmdList, ID3D11DeviceContext* Direct3DDeviceContext,FParameter)
 		{
 		}
 
@@ -93,7 +93,7 @@ namespace
 		}
 		FResolveSingleSamplePS() {}
 
-		void SetParameters(FRHICommandList* RHICmdList, ID3D11DeviceContext* Direct3DDeviceContext,uint32 SingleSampleIndexValue)
+		void SetParameters(FRHICommandList& RHICmdList, ID3D11DeviceContext* Direct3DDeviceContext,uint32 SingleSampleIndexValue)
 		{
 			SetShaderValue(RHICmdList, GetPixelShader(),SingleSampleIndex,SingleSampleIndexValue);
 		}
@@ -173,7 +173,7 @@ void FD3D11DynamicRHI::ResolveTextureUsingShader(
 	)
 {
 	//@todo-rco: RHIPacketList
-	FRHICommandList* RHICmdList = nullptr;
+	FRHICommandList& RHICmdList = FRHICommandList::GetNullRef();
 
 	// Save the current viewport so that it can be restored
 	D3D11_VIEWPORT SavedViewport;
