PS4 - Geometry shader ring buffers
* Off chip GS ring buffers are now handled by allocating the global resource table from the command buffer when required. Removes the need to use the largets output for all GS.
* Initial support for OnChip GS. Internal compiler not currently supported due to missing profile. Also seeing some random polys which looks to be an issue with the SDK.

diff --git a/Engine/Source/Runtime/ShaderCore/Public/ShaderCore.h b/Engine/Source/Runtime/ShaderCore/Public/ShaderCore.h
--- a/Engine/Source/Runtime/ShaderCore/Public/ShaderCore.h
+++ b/Engine/Source/Runtime/ShaderCore/Public/ShaderCore.h
@@ -108,7 +108,9 @@ enum ECompilerFlags
 	/** Disable shader validation */
 	CFLAG_SkipValidation,
 	/** Only allows standard optimizations, not the longest compile times. */
-	CFLAG_StandardOptimization
+	CFLAG_StandardOptimization,
+	/** Shader should use on chip memory instead of main memory ring buffer memory. */
+	CFLAG_OnChip
 };
 
 /**
