UE4 - remove stale comments

diff --git a/Engine/Source/Runtime/OpenGLDrv/Private/OpenGLDevice.cpp b/Engine/Source/Runtime/OpenGLDrv/Private/OpenGLDevice.cpp
--- a/Engine/Source/Runtime/OpenGLDrv/Private/OpenGLDevice.cpp
+++ b/Engine/Source/Runtime/OpenGLDrv/Private/OpenGLDevice.cpp
@@ -65,7 +65,6 @@ void OnPixelBufferDeletion( GLuint PixelBufferResource )
 	PrivateOpenGLDevicePtr->OnPixelBufferDeletion( PixelBufferResource );
 }
 
-//gilmerge lose more of these args?
 void OnUniformBufferDeletion( GLuint UniformBufferResource, uint32 AllocatedSize, bool bStreamDraw, uint32 , uint8*  )
 {
 	check(PrivateOpenGLDevicePtr);
diff --git a/Engine/Source/Runtime/OpenGLDrv/Private/OpenGLUniformBuffer.cpp b/Engine/Source/Runtime/OpenGLDrv/Private/OpenGLUniformBuffer.cpp
--- a/Engine/Source/Runtime/OpenGLDrv/Private/OpenGLUniformBuffer.cpp
+++ b/Engine/Source/Runtime/OpenGLDrv/Private/OpenGLUniformBuffer.cpp
@@ -513,7 +513,6 @@ FUniformBufferRHIRef FOpenGLDynamicRHI::RHICreateUniformBuffer(const void* Conte
 			}
 			else if (IsSuballocatingUBOs())
 			{
-				//gilmerge was SuballocateUBO(BufferSize, AllocatedResource, OffsetInBuffer, PersistentlyMappedBuffer);
 				SuballocateUBO(AllocatedSize, AllocatedResource, OffsetInBuffer, PersistentlyMappedBuffer);
 			}
 			else
diff --git a/Engine/Source/Runtime/Renderer/Private/PostProcess/SceneRenderTargets.cpp b/Engine/Source/Runtime/Renderer/Private/PostProcess/SceneRenderTargets.cpp
--- a/Engine/Source/Runtime/Renderer/Private/PostProcess/SceneRenderTargets.cpp
+++ b/Engine/Source/Runtime/Renderer/Private/PostProcess/SceneRenderTargets.cpp
@@ -1717,7 +1717,6 @@ void FDeferredPixelShaderParameters::Set(FRHICommandList& RHICmdList, const Shad
 
 		if (FeatureLevel >= ERHIFeatureLevel::SM4)
 		{
-			//gilmerge
 			if (GBufferResources.IsBound())
 			{
 				SetUniformBufferParameter(RHICmdList, ShaderRHI, GBufferResources, GSceneRenderTargets.GetGBufferResourcesUniformBuffer());
diff --git a/Engine/Source/Runtime/ShaderCore/Public/ShaderCore.h b/Engine/Source/Runtime/ShaderCore/Public/ShaderCore.h
--- a/Engine/Source/Runtime/ShaderCore/Public/ShaderCore.h
+++ b/Engine/Source/Runtime/ShaderCore/Public/ShaderCore.h
@@ -417,7 +417,6 @@ struct FShaderCompilerInput
 	FString SourceFilePrefix;
 	FString SourceFilename;
 	FString EntryPointName;
-	//gilmerge, it isn't clear what the relationship is between these two and what the call sites are supposed to use. I might have botched the call sites.
 	FString DumpDebugInfoRootPath;	// Dump debug path (up to platform)
 	FString DumpDebugInfoPath;		// Dump debug path (platform/groupname)
 	FShaderCompilerEnvironment Environment;
