Disabling the code in a way that doesn't trigger compiler warnings.

diff --git a/Engine/Source/Runtime/ShaderCore/Private/Shader.cpp b/Engine/Source/Runtime/ShaderCore/Private/Shader.cpp
--- a/Engine/Source/Runtime/ShaderCore/Private/Shader.cpp
+++ b/Engine/Source/Runtime/ShaderCore/Private/Shader.cpp
@@ -909,7 +909,7 @@ void FShader::VerifyBoundUniformBufferParameters()
 	// Support being called on a NULL pointer
 // TODO: doesn't work with uniform buffer parameters on helper structs like FDeferredPixelShaderParameters
 	//@todo parallelrendering
-	if (0&&this)
+	if (0)//&&this)
 	{
 		for (int32 StructIndex = 0; StructIndex < UniformBufferParameters.Num(); StructIndex++)
 		{
