Better error message for shader / VF types loaded too late

diff --git a/Engine/Source/Runtime/ShaderCore/Private/Shader.cpp b/Engine/Source/Runtime/ShaderCore/Private/Shader.cpp
--- a/Engine/Source/Runtime/ShaderCore/Private/Shader.cpp
+++ b/Engine/Source/Runtime/ShaderCore/Private/Shader.cpp
@@ -90,7 +90,7 @@ FShaderType::FShaderType(
 
 	// This will trigger if an IMPLEMENT_SHADER_TYPE was in a module not loaded before InitializeShaderTypes
 	// Shader types need to be implemented in modules that are loaded before that
-	check(!bInitializedSerializationHistory);
+	checkf(!bInitializedSerializationHistory, TEXT("Shader type was loaded after engine init, use ELoadingPhase::PostConfigInit on your module to cause it to load earlier."));
 
 	//make sure the name is shorter than the maximum serializable length
 	check(FCString::Strlen(InName) < NAME_SIZE);
diff --git a/Engine/Source/Runtime/ShaderCore/Private/VertexFactory.cpp b/Engine/Source/Runtime/ShaderCore/Private/VertexFactory.cpp
--- a/Engine/Source/Runtime/ShaderCore/Private/VertexFactory.cpp
+++ b/Engine/Source/Runtime/ShaderCore/Private/VertexFactory.cpp
@@ -118,7 +118,7 @@ FVertexFactoryType::FVertexFactoryType(
 
 	// This will trigger if an IMPLEMENT_VERTEX_FACTORY_TYPE was in a module not loaded before InitializeShaderTypes
 	// Vertex factory types need to be implemented in modules that are loaded before that
-	check(!bInitializedSerializationHistory);
+	checkf(!bInitializedSerializationHistory, TEXT("VF type was loaded after engine init, use ELoadingPhase::PostConfigInit on your module to cause it to load earlier."));
 
 	// Add this vertex factory type to the global list.
 	GlobalListLink.Link(GetTypeList());
