Removed excess logging when shader platforms are not compatible.

diff --git a/Engine/Source/Runtime/ShaderCore/Private/Shader.cpp b/Engine/Source/Runtime/ShaderCore/Private/Shader.cpp
--- a/Engine/Source/Runtime/ShaderCore/Private/Shader.cpp
+++ b/Engine/Source/Runtime/ShaderCore/Private/Shader.cpp
@@ -432,20 +432,9 @@ void FShaderResource::InitRHI()
 	// we can't have this called on the wrong platform's shaders
 	if (!ArePlatformsCompatible(GMaxRHIShaderPlatform, (EShaderPlatform)Target.Platform))
  	{
-		for (int i = 0; i < EShaderPlatform::SP_NumPlatforms; i++)
-		{
-			for (int j = 0; j < EShaderPlatform::SP_NumPlatforms; j++)
-			{
-				bool bOK = ArePlatformsCompatible((EShaderPlatform)i, (EShaderPlatform)j);
-
-				UE_LOG(LogShaders, Warning, TEXT("platform %s %s %s"),
-					*LegacyShaderPlatformToShaderFormat((EShaderPlatform)i).ToString(), (bOK ? *FString("IS compatible with") : *FString("ISNOT compatible with")), *LegacyShaderPlatformToShaderFormat((EShaderPlatform)j).ToString());
-			}
-		}
-
  		if (FPlatformProperties::RequiresCookedData())
  		{
- 			UE_LOG(LogShaders, Fatal, TEXT("FShaderResource::InitRHI got platform %s but expected %s"), 
+ 			UE_LOG(LogShaders, Fatal, TEXT("FShaderResource::InitRHI got platform %s but it is not compatible with %s"), 
  				*LegacyShaderPlatformToShaderFormat((EShaderPlatform)Target.Platform).ToString(), *LegacyShaderPlatformToShaderFormat(GRHIShaderPlatform).ToString() );
  		}
  		return;
