comment

diff --git a/Engine/Source/Runtime/ShaderCore/Public/Shader.h b/Engine/Source/Runtime/ShaderCore/Public/Shader.h
--- a/Engine/Source/Runtime/ShaderCore/Public/Shader.h
+++ b/Engine/Source/Runtime/ShaderCore/Public/Shader.h
@@ -843,7 +843,7 @@ private:
 		ShaderClass::ShouldCache \
 		);
 
-/** A macro to implement a shader type, the function name and the source filename comes from the class. */
+/** A macro to implement a templated shader type, the function name and the source filename comes from the class. */
 #define IMPLEMENT_SHADER_TYPE2(ShaderClass,Frequency) \
 	template<> \
 	ShaderClass::ShaderMetaType ShaderClass::StaticType( \
