UE4 - Remove unused RHI globals:
- GPixelCenterOffset
- GSupportsVertexInstancing
- GSupportsEmulatedVertexInstancing
- GVertexElementsCanShareStreamOffset
- GSupportsVertexTextureFetch

diff --git a/Engine/Source/Runtime/EmptyRHI/Private/EmptyRHI.cpp b/Engine/Source/Runtime/EmptyRHI/Private/EmptyRHI.cpp
--- a/Engine/Source/Runtime/EmptyRHI/Private/EmptyRHI.cpp
+++ b/Engine/Source/Runtime/EmptyRHI/Private/EmptyRHI.cpp
@@ -28,10 +28,6 @@ FEmptyDynamicRHI::FEmptyDynamicRHI()
 
 	// Initialize the RHI capabilities.
 
-// 	GPixelCenterOffset = 0.0f;
-// 	GSupportsVertexInstancing = false;
-// 	GSupportsEmulatedVertexInstancing = false;
-// 	GSupportsVertexTextureFetch = false;
 //	GRHIAdapterName = 
 //	GRHIVendorId = 
 // 	GSupportsRenderTargetFormat_PF_G8 = false;
diff --git a/Engine/Source/Runtime/Engine/Private/UserInterface/Canvas.cpp b/Engine/Source/Runtime/Engine/Private/UserInterface/Canvas.cpp
--- a/Engine/Source/Runtime/Engine/Private/UserInterface/Canvas.cpp
+++ b/Engine/Source/Runtime/Engine/Private/UserInterface/Canvas.cpp
@@ -104,7 +104,7 @@ FMatrix FCanvas::CalcBaseTransform2D(uint32 ViewSizeX, uint32 ViewSizeY)
 	ViewSizeY = FMath::Max<uint32>(ViewSizeY, 1.f);
 
 	return AdjustProjectionMatrixForRHI(
-		FTranslationMatrix(FVector(-GPixelCenterOffset,-GPixelCenterOffset,0)) *
+		FTranslationMatrix(FVector(0, 0, 0)) *
 		FMatrix(
 			FPlane(	1.0f / (ViewSizeX / 2.0f),	0.0,										0.0f,	0.0f	),
 			FPlane(	0.0f,						-1.0f / (ViewSizeY / 2.0f),					0.0f,	0.0f	),
@@ -139,7 +139,7 @@ FMatrix FCanvas::CalcViewMatrix(uint32 ViewSizeX, uint32 ViewSizeY, float fFOV)
 		);
 	return 
 		// also apply screen offset to align to pixel centers
-		FTranslationMatrix(FVector(-GPixelCenterOffset,-GPixelCenterOffset,0)) * 
+		FTranslationMatrix(FVector(0, 0, 0)) * 
 		OrientCanvasMat * 
 		CamOffsetMat;
 }
diff --git a/Engine/Source/Runtime/IOS/MetalRHI/Private/MetalRHI.cpp b/Engine/Source/Runtime/IOS/MetalRHI/Private/MetalRHI.cpp
--- a/Engine/Source/Runtime/IOS/MetalRHI/Private/MetalRHI.cpp
+++ b/Engine/Source/Runtime/IOS/MetalRHI/Private/MetalRHI.cpp
@@ -66,10 +66,6 @@ FMetalDynamicRHI::FMetalDynamicRHI()
 		GEmitDrawEvents = true;
 	}
 
-	GPixelCenterOffset = 0.0f;
-	GSupportsVertexInstancing = false;
-	GSupportsEmulatedVertexInstancing = false;
-	GSupportsVertexTextureFetch = false;
 	GSupportsShaderFramebufferFetch = true;
 	GHardwareHiddenSurfaceRemoval = true;
 //	GRHIAdapterName = 
diff --git a/Engine/Source/Runtime/OpenGLDrv/Private/OpenGLDevice.cpp b/Engine/Source/Runtime/OpenGLDrv/Private/OpenGLDevice.cpp
--- a/Engine/Source/Runtime/OpenGLDrv/Private/OpenGLDevice.cpp
+++ b/Engine/Source/Runtime/OpenGLDrv/Private/OpenGLDevice.cpp
@@ -650,7 +650,6 @@ static void InitRHICapabilitiesForGL()
 
 	GSupportsVolumeTextureRendering = FOpenGL::SupportsVolumeTextureRendering();
 	GSupportsRenderDepthTargetableShaderResources = true;
-	GSupportsVertexTextureFetch = true;
 	GSupportsRenderTargetFormat_PF_G8 = true;
 	GSupportsSeparateRenderTargetBlendState = FOpenGL::SupportsSeparateAlphaBlend();
 	GSupportsDepthBoundsTest = FOpenGL::SupportsDepthBoundsTest();
@@ -658,10 +657,8 @@ static void InitRHICapabilitiesForGL()
 	GSupportsRenderTargetFormat_PF_FloatRGBA = FOpenGL::SupportsColorBufferHalfFloat();
 
 	GSupportsShaderFramebufferFetch = FOpenGL::SupportsShaderFramebufferFetch();
-	GPixelCenterOffset = 0.0f;
 	GMaxShadowDepthBufferSizeX = FMath::Min<int32>(Value_GL_MAX_RENDERBUFFER_SIZE, 4096); // Limit to the D3D11 max.
 	GMaxShadowDepthBufferSizeY = FMath::Min<int32>(Value_GL_MAX_RENDERBUFFER_SIZE, 4096);
-	GSupportsVertexInstancing = true;
 	GHardwareHiddenSurfaceRemoval = FOpenGL::HasHardwareHiddenSurfaceRemoval();
 
 	GShaderPlatformForFeatureLevel[ERHIFeatureLevel::ES2] = (GMaxRHIFeatureLevel == ERHIFeatureLevel::ES2) ? GMaxRHIShaderPlatform : SP_OPENGL_PCES2;
diff --git a/Engine/Source/Runtime/RHI/Private/RHI.cpp b/Engine/Source/Runtime/RHI/Private/RHI.cpp
--- a/Engine/Source/Runtime/RHI/Private/RHI.cpp
+++ b/Engine/Source/Runtime/RHI/Private/RHI.cpp
@@ -159,7 +159,6 @@ namespace RHIConfig
 
 bool GIsRHIInitialized = false;
 int32 GMaxTextureMipCount = MAX_TEXTURE_MIP_COUNT;
-bool GSupportsVertexTextureFetch = false;
 bool GSupportsDepthFetchDuringDepthTest = true;
 FString GRHIAdapterName;
 uint32 GRHIVendorId = 0;
@@ -173,7 +172,6 @@ bool GSupportsQuads = false;
 bool GSupportsVolumeTextureRendering = true;
 bool GSupportsSeparateRenderTargetBlendState = false;
 bool GSupportsDepthRenderTargetWithoutColorRenderTarget = true;
-float GPixelCenterOffset = 0;
 float GMinClipZ = 0.0f;
 float GProjectionSignY = 1.0f;
 int32 GMaxShadowDepthBufferSizeX = 2048;
@@ -184,9 +182,6 @@ int32 GMaxTextureArrayLayers = 256;
 bool GUsingNullRHI = false;
 int32 GDrawUPVertexCheckCount = MAX_int32;
 int32 GDrawUPIndexCheckCount = MAX_int32;
-bool GSupportsVertexInstancing = false;
-bool GSupportsEmulatedVertexInstancing = false;
-bool GVertexElementsCanShareStreamOffset = true;
 bool GTriggerGPUProfile = false;
 bool GRHISupportsTextureStreaming = false;
 bool GSupportsDepthBoundsTest = false;
diff --git a/Engine/Source/Runtime/RHI/Public/RHI.h b/Engine/Source/Runtime/RHI/Public/RHI.h
--- a/Engine/Source/Runtime/RHI/Public/RHI.h
+++ b/Engine/Source/Runtime/RHI/Public/RHI.h
@@ -52,11 +52,6 @@ extern	RHI_API int32		GMaxTextureMipCount;
 /** true if the RHI supports textures that may be bound as both a render target and a shader resource. */
 extern RHI_API bool GSupportsRenderDepthTargetableShaderResources;
 
-/**
-* true if the our renderer, the driver and the hardware supports the feature.
-*/
-extern RHI_API bool GSupportsVertexTextureFetch;
-
 /** true if the RHI supports binding depth as a texture when testing against depth */
 extern RHI_API bool GSupportsDepthFetchDuringDepthTest;
 
@@ -108,9 +103,6 @@ extern RHI_API bool GSupportsDepthRenderTargetWithoutColorRenderTarget;
 /** True if the RHI supports depth bounds testing */
 extern RHI_API bool GSupportsDepthBoundsTest;
 
-/** The offset from the upper left corner of a pixel to the position which is used to sample textures for fragments of that pixel. */
-extern RHI_API float GPixelCenterOffset;
-
 /** The minimum Z value in clip space for the RHI. */
 extern RHI_API float GMinClipZ;
 
@@ -157,15 +149,6 @@ extern RHI_API int32 GDrawUPVertexCheckCount;
  */
 extern RHI_API int32 GDrawUPIndexCheckCount;
 
-/** true if the rendering hardware supports vertex instancing. */
-extern RHI_API bool GSupportsVertexInstancing;
-
-/** true if the rendering hardware can emulate vertex instancing. */
-extern RHI_API bool GSupportsEmulatedVertexInstancing;
-
-/** If false code needs to patch up vertex declaration. */
-extern RHI_API bool GVertexElementsCanShareStreamOffset;
-
 /** true for each VET that is supported. One-to-one mapping with EVertexElementType */
 extern RHI_API class FVertexElementTypeSupportInfo GVertexElementTypeSupport;
 
diff --git a/Engine/Source/Runtime/Renderer/Private/CompositionLighting/PostProcessAmbientOcclusion.h b/Engine/Source/Runtime/Renderer/Private/CompositionLighting/PostProcessAmbientOcclusion.h
--- a/Engine/Source/Runtime/Renderer/Private/CompositionLighting/PostProcessAmbientOcclusion.h
+++ b/Engine/Source/Runtime/Renderer/Private/CompositionLighting/PostProcessAmbientOcclusion.h
@@ -96,8 +96,8 @@ public:
 		FVector4 AOScreenPositionScaleBiasValue = FVector4(
 			ViewRect.Size().X / InputTextureSize.X / +2.0f,
 			ViewRect.Size().Y / InputTextureSize.Y / -2.0f,
-			(ViewRect.Size().X / 2.0f + GPixelCenterOffset + ViewRect.Min.X) / InputTextureSize.X,
-			(ViewRect.Size().Y / 2.0f + GPixelCenterOffset + ViewRect.Min.Y) / InputTextureSize.Y);
+			(ViewRect.Size().X / 2.0f + ViewRect.Min.X) / InputTextureSize.X,
+			(ViewRect.Size().Y / 2.0f + ViewRect.Min.Y) / InputTextureSize.Y);
 
 		// Combining two scales into one parameter, Projection matrix scaling of x and y and scaling from clip to screen space.
 		FVector2D ProjectionScaleValue = FVector2D(View.ViewMatrices.ProjMatrix.M[0][0], View.ViewMatrices.ProjMatrix.M[1][1])
diff --git a/Engine/Source/Runtime/Renderer/Private/DistanceFieldSurfaceCacheLighting.cpp b/Engine/Source/Runtime/Renderer/Private/DistanceFieldSurfaceCacheLighting.cpp
--- a/Engine/Source/Runtime/Renderer/Private/DistanceFieldSurfaceCacheLighting.cpp
+++ b/Engine/Source/Runtime/Renderer/Private/DistanceFieldSurfaceCacheLighting.cpp
@@ -3124,8 +3124,8 @@ void RenderIrradianceCacheInterpolation(
 				View.ViewRect.Height() * InvBufferSizeY / (-2.0f * GProjectionSignY));
 
 			const FVector2D ScreenPositionBias(
-				(View.ViewRect.Width() / 2.0f + GPixelCenterOffset) * InvBufferSizeX,
-				(View.ViewRect.Height() / 2.0f + GPixelCenterOffset) * InvBufferSizeY);
+				(View.ViewRect.Width() / 2.0f) * InvBufferSizeX,
+				(View.ViewRect.Height() / 2.0f) * InvBufferSizeY);
 
 			FVector2D OffsetToLowResCorner = (FVector2D(.5f, .5f) / FVector2D(DownsampledViewSizeX, DownsampledViewSizeY) - FVector2D(.5f, .5f)) * FVector2D(2, -2);
 			FVector2D OffsetToTopResPixel = (FVector2D(.5f, .5f) / BufferSize - ScreenPositionBias) / ScreenPositionScale;
diff --git a/Engine/Source/Runtime/Renderer/Private/PostProcess/SceneFilterRendering.cpp b/Engine/Source/Runtime/Renderer/Private/PostProcess/SceneFilterRendering.cpp
--- a/Engine/Source/Runtime/Renderer/Private/PostProcess/SceneFilterRendering.cpp
+++ b/Engine/Source/Runtime/Renderer/Private/PostProcess/SceneFilterRendering.cpp
@@ -289,8 +289,8 @@ void DrawTransformedRectangle(
 
 	for (int32 VertexIndex = 0; VertexIndex < 4; VertexIndex++)
 	{
-		Vertices[VertexIndex].Position.X = -1.0f + 2.0f * (Vertices[VertexIndex].Position.X - GPixelCenterOffset) / (float)TargetSize.X;
-		Vertices[VertexIndex].Position.Y = (+1.0f - 2.0f * (Vertices[VertexIndex].Position.Y - GPixelCenterOffset) / (float)TargetSize.Y) * GProjectionSignY;
+		Vertices[VertexIndex].Position.X = -1.0f + 2.0f * (Vertices[VertexIndex].Position.X) / (float)TargetSize.X;
+		Vertices[VertexIndex].Position.Y = (+1.0f - 2.0f * (Vertices[VertexIndex].Position.Y) / (float)TargetSize.Y) * GProjectionSignY;
 
 		Vertices[VertexIndex].UV.X = Vertices[VertexIndex].UV.X / (float)TextureSize.X;
 		Vertices[VertexIndex].UV.Y = Vertices[VertexIndex].UV.Y / (float)TextureSize.Y;
diff --git a/Engine/Source/Runtime/Renderer/Private/SceneHitProxyRendering.cpp b/Engine/Source/Runtime/Renderer/Private/SceneHitProxyRendering.cpp
--- a/Engine/Source/Runtime/Renderer/Private/SceneHitProxyRendering.cpp
+++ b/Engine/Source/Runtime/Renderer/Private/SceneHitProxyRendering.cpp
@@ -532,7 +532,7 @@ void RenderHitProxies(FRHICommandListImmediate& RHICmdList, const FSceneRenderer
 	FIntPoint RenderTargetSize = ViewFamily.RenderTarget->GetSizeXY();
 
 	const FMatrix PixelToView =
-		FTranslationMatrix(FVector(-GPixelCenterOffset,-GPixelCenterOffset,0)) *
+		FTranslationMatrix(FVector(0, 0, 0)) *
 			FMatrix(
 				FPlane(	1.0f / ((float)RenderTargetSize.X / 2.0f),	0.0,										0.0f,	0.0f	),
 				FPlane(	0.0f,										-GProjectionSignY / ((float)RenderTargetSize.Y / 2.0f),	0.0f,	0.0f	),
diff --git a/Engine/Source/Runtime/Renderer/Private/SceneRendering.cpp b/Engine/Source/Runtime/Renderer/Private/SceneRendering.cpp
--- a/Engine/Source/Runtime/Renderer/Private/SceneRendering.cpp
+++ b/Engine/Source/Runtime/Renderer/Private/SceneRendering.cpp
@@ -289,8 +289,8 @@ TUniformBufferRef<FViewUniformShaderParameters> FViewInfo::CreateUniformBuffer(
 	const FVector4 ScreenPositionScaleBias(
 		ViewRect.Width() * InvBufferSizeX / +2.0f,
 		ViewRect.Height() * InvBufferSizeY / (-2.0f * GProjectionSignY),
-		(ViewRect.Height() / 2.0f + GPixelCenterOffset + ViewRect.Min.Y) * InvBufferSizeY,
-		(ViewRect.Width() / 2.0f + GPixelCenterOffset + ViewRect.Min.X) * InvBufferSizeX
+		(ViewRect.Height() / 2.0f + ViewRect.Min.Y) * InvBufferSizeY,
+		(ViewRect.Width() / 2.0f + ViewRect.Min.X) * InvBufferSizeX
 		);
 	
 	const bool bIsUnlitView = !Family->EngineShowFlags.Lighting;
diff --git a/Engine/Source/Runtime/Renderer/Private/ShadowRendering.cpp b/Engine/Source/Runtime/Renderer/Private/ShadowRendering.cpp
--- a/Engine/Source/Runtime/Renderer/Private/ShadowRendering.cpp
+++ b/Engine/Source/Runtime/Renderer/Private/ShadowRendering.cpp
@@ -2344,8 +2344,8 @@ FMatrix FProjectedShadowInfo::GetScreenToShadowMatrix(const FSceneView& View) co
 			FPlane(0,						 -ShadowResolutionFractionY,0,									0),
 			FPlane(0,						0,							InvMaxSubjectDepth,	0),
 			FPlane(
-				(X + SHADOW_BORDER + GPixelCenterOffset) * InvBufferResolutionX + ShadowResolutionFractionX,
-				(Y + SHADOW_BORDER + GPixelCenterOffset) * InvBufferResolutionY + ShadowResolutionFractionY,
+				(X + SHADOW_BORDER) * InvBufferResolutionX + ShadowResolutionFractionX,
+				(Y + SHADOW_BORDER) * InvBufferResolutionY + ShadowResolutionFractionY,
 				0,
 				1
 			)
@@ -2376,8 +2376,8 @@ FMatrix FProjectedShadowInfo::GetWorldToShadowMatrix(FVector4& ShadowmapMinMax,
 			FPlane(0,						 -ShadowResolutionFractionY,0,									0),
 			FPlane(0,						0,							InvMaxSubjectDepth,	0),
 			FPlane(
-				(X + ShadowBorder + GPixelCenterOffset) * InvBufferResolutionX + ShadowResolutionFractionX,
-				(Y + ShadowBorder + GPixelCenterOffset) * InvBufferResolutionY + ShadowResolutionFractionY,
+				(X + ShadowBorder) * InvBufferResolutionX + ShadowResolutionFractionX,
+				(Y + ShadowBorder) * InvBufferResolutionY + ShadowResolutionFractionY,
 				0,
 				1
 			)
diff --git a/Engine/Source/Runtime/ShaderCore/Private/VertexFactory.cpp b/Engine/Source/Runtime/ShaderCore/Private/VertexFactory.cpp
--- a/Engine/Source/Runtime/ShaderCore/Private/VertexFactory.cpp
+++ b/Engine/Source/Runtime/ShaderCore/Private/VertexFactory.cpp
@@ -289,59 +289,6 @@ static uint8 GetVertexElementSize( uint8 Type )
 	}
 }
 
-/**
- * Patches the declaration so vertex stream offsets are unique. This is required for e.g. GeForce FX cards, which don't support redundant 
- * offsets in the declaration. We're unable to make that many vertex elements point to the same offset so the function moves redundant 
- * declarations to higher offsets, pointing to garbage data.
- */
-static void PatchVertexStreamOffsetsToBeUnique( FVertexDeclarationElementList& Elements )
-{
-	// check every vertex element
-	for ( int32 e = 0; e < Elements.Num(); e++ )
-	{
-		// check if there's an element that reads from the same offset
-		for ( int32 i = 0; i < Elements.Num(); i++ )
-		{
-			// but only in the same stream and if it's not the same element
-			if ( ( Elements[ i ].StreamIndex == Elements[ e ].StreamIndex ) && ( Elements[ i ].Offset == Elements[ e ].Offset ) && ( e != i ) )
-			{
-				// the id of the highest offset element is stored here (it doesn't need to be the last element in the declarator because the last element may belong to another StreamIndex
-				uint32 MaxOffsetID = i;
-
-				// find the highest offset element
-				for ( int32 j = 0; j < Elements.Num(); j++ )
-				{
-					if ( ( Elements[ j ].StreamIndex == Elements[ e ].StreamIndex ) && ( Elements[ MaxOffsetID ].Offset < Elements[ j ].Offset ) )
-					{
-						MaxOffsetID = j;
-					}
-				}
-
-				// get the size of the highest offset element, it's needed for the redundant element new offset
-				uint8 PreviousElementSize = GetVertexElementSize( Elements[ MaxOffsetID ].Type );
-
-				// prepare a new vertex element
-				FVertexElement VertElement;
-				VertElement.Offset		= Elements[ MaxOffsetID ].Offset + PreviousElementSize;
-				VertElement.StreamIndex = Elements[ i ].StreamIndex;
-				VertElement.Type		= Elements[ i ].Type;
-				VertElement.AttributeIndex	= Elements[ i ].AttributeIndex;
-				VertElement.Stride = Elements[ i ].Stride;
-				VertElement.bUseInstanceIndex = Elements[i].bUseInstanceIndex;
-
-				// remove the old redundant element
-				Elements.RemoveAt( i );
-
-				// add a new element with "correct" offset
-				Elements.Add( VertElement );
-
-				// make sure that when the element has been removed its index is taken by the next element, so we must take care of it too
-				i = i == 0 ? 0 : i - 1;
-			}
-		}
-	}
-}
-
 void FVertexFactory::InitDeclaration(
 	FVertexDeclarationElementList& Elements, 
 	const DataType& InData)
@@ -349,12 +296,6 @@ void FVertexFactory::InitDeclaration(
 	// Make a copy of the vertex factory data.
 	Data = InData;
 
-	// If GFFX detected, patch up the declarator
-	if( !GVertexElementsCanShareStreamOffset )
-	{
-		PatchVertexStreamOffsetsToBeUnique( Elements );
-	}
-
 	// Create the vertex declaration for rendering the factory normally.
 	Declaration = RHICreateVertexDeclaration(Elements);
 }
diff --git a/Engine/Source/Runtime/SlateRHIRenderer/Private/SlateRHIRenderer.cpp b/Engine/Source/Runtime/SlateRHIRenderer/Private/SlateRHIRenderer.cpp
--- a/Engine/Source/Runtime/SlateRHIRenderer/Private/SlateRHIRenderer.cpp
+++ b/Engine/Source/Runtime/SlateRHIRenderer/Private/SlateRHIRenderer.cpp
@@ -130,13 +130,10 @@ FSlateRHIRenderer::~FSlateRHIRenderer()
 
 FMatrix FSlateRHIRenderer::CreateProjectionMatrix(uint32 Width, uint32 Height)
 {
-
-	float PixelOffset = GPixelCenterOffset;
-
 	// Create ortho projection matrix
-	const float Left = PixelOffset;
+	const float Left = 0;
 	const float Right = Left + Width;
-	const float Top = PixelOffset;
+	const float Top = 0;
 	const float Bottom = Top + Height;
 	const float ZNear = -100.0f;
 	const float ZFar = 100.0f;
diff --git a/Engine/Source/Runtime/SlateRHIRenderer/Private/SlateRHIRenderingPolicy.cpp b/Engine/Source/Runtime/SlateRHIRenderer/Private/SlateRHIRenderingPolicy.cpp
--- a/Engine/Source/Runtime/SlateRHIRenderer/Private/SlateRHIRenderingPolicy.cpp
+++ b/Engine/Source/Runtime/SlateRHIRenderer/Private/SlateRHIRenderingPolicy.cpp
@@ -89,7 +89,7 @@ void FSlateElementIndexBuffer::ReleaseDynamicRHI()
 }
 
 FSlateRHIRenderingPolicy::FSlateRHIRenderingPolicy( TSharedPtr<FSlateFontCache> InFontCache, TSharedRef<FSlateRHIResourceManager> InResourceManager )
-	: FSlateRenderingPolicy( GPixelCenterOffset )
+	: FSlateRenderingPolicy(0)
 	, ResourceManager( InResourceManager )
 	, FontCache( InFontCache )
 	, CurrentBufferIndex(0)
diff --git a/Engine/Source/Runtime/Windows/D3D11RHI/Private/D3D11Device.cpp b/Engine/Source/Runtime/Windows/D3D11RHI/Private/D3D11Device.cpp
--- a/Engine/Source/Runtime/Windows/D3D11RHI/Private/D3D11Device.cpp
+++ b/Engine/Source/Runtime/Windows/D3D11RHI/Private/D3D11Device.cpp
@@ -89,10 +89,6 @@ FD3D11DynamicRHI::FD3D11DynamicRHI(IDXGIFactory* InDXGIFactory,D3D_FEATURE_LEVEL
 		GMaxRHIShaderPlatform = SP_PCD3D_SM4;
 	}
 
-	GPixelCenterOffset = 0.0f;	// Note that in D3D11, there is no half-texel offset (ala DX9)	
-	GSupportsVertexInstancing = true;
-	GSupportsVertexTextureFetch = true;
-
 	// Initialize the platform pixel format map.
 	GPixelFormats[ PF_Unknown		].PlatformFormat	= DXGI_FORMAT_UNKNOWN;
 	GPixelFormats[ PF_A32B32G32R32F	].PlatformFormat	= DXGI_FORMAT_R32G32B32A32_FLOAT;
