Initial ShaderCore-level cache to start addressing shader hitches [UE-7536]. At present it is only configured by [REDACTED_PERSON], though other APIs could use it. As all shaders & programs are cached it is inappropriate for memory constrained devices so it is only enabled by Desktop platforms.
   - The consolve variable "r.UseShaderCaching" controls whether the cache should be used.
   - The cache has to be initialised and shutdown by the RHI.
   - All RHI shader types now contain the SHA hash of the compiled source when the cache is in use so that they can be identified efficiently.
   - All shaders are fetched from the cache rather than being created on-demand.
   - On shader deserialisation the new shader is sent to the render thread for creation & caching.
   - After each shader is cached all complete bound-shader-states that include that shader are constructed and cached.
   - Wrapped OpenGL glCreateShader & glCreateProgram to use name-caching as these operations can synchronise the entire OpenGL pipeline on OS X.

diff --git a/Engine/Source/Runtime/OpenGLDrv/Private/OpenGLDevice.cpp b/Engine/Source/Runtime/OpenGLDrv/Private/OpenGLDevice.cpp
--- a/Engine/Source/Runtime/OpenGLDrv/Private/OpenGLDevice.cpp
+++ b/Engine/Source/Runtime/OpenGLDrv/Private/OpenGLDevice.cpp
@@ -7,6 +7,7 @@
 #include "OpenGLDrvPrivate.h"
 
 #include "HardwareInfo.h"
+#include "ShaderCache.h"
 
 extern GLint GMaxOpenGLColorSamples;
 extern GLint GMaxOpenGLDepthSamples;
@@ -1044,6 +1045,10 @@ void FOpenGLDynamicRHI::Init()
 	check(!GIsRHIInitialized);
 	VERIFY_GL_SCOPE();
 
+#if PLATFORM_DESKTOP
+	FShaderCache::InitShaderCache();
+#endif
+
 	InitializeStateResources();
 
 	// Create a default point sampler state for internal use.
@@ -1118,6 +1123,10 @@ void FOpenGLDynamicRHI::Cleanup()
 {
 	if(GIsRHIInitialized)
 	{
+#if PLATFORM_DESKTOP
+		FShaderCache::ShutdownShaderCache();
+#endif
+
 		// Reset the RHI initialized flag.
 		GIsRHIInitialized = false;
 
diff --git a/Engine/Source/Runtime/OpenGLDrv/Private/OpenGLShaders.cpp b/Engine/Source/Runtime/OpenGLDrv/Private/OpenGLShaders.cpp
--- a/Engine/Source/Runtime/OpenGLDrv/Private/OpenGLShaders.cpp
+++ b/Engine/Source/Runtime/OpenGLDrv/Private/OpenGLShaders.cpp
@@ -7,6 +7,7 @@
 #include "OpenGLDrvPrivate.h"
 #include "Shader.h"
 #include "GlobalShader.h"
+#include "ShaderCache.h"
 
 #define CHECK_FOR_GL_SHADERS_TO_REPLACE 0
 
@@ -46,6 +47,7 @@ static bool VerifyCompiledShader(GLuint Shader, const ANSICHAR* GlslCode )
 {
 	SCOPE_CYCLE_COUNTER(STAT_OpenGLShaderCompileVerifyTime);
 
+#if UE_BUILD_DEBUG
 	GLint CompileStatus;
 	glGetShaderiv(Shader, GL_COMPILE_STATUS, &CompileStatus);
 	if (CompileStatus != GL_TRUE)
@@ -96,6 +98,7 @@ static bool VerifyCompiledShader(GLuint Shader, const ANSICHAR* GlslCode )
 		}
 		return false;
 	}
+#endif
 	return true;
 }
 
@@ -106,6 +109,7 @@ static bool VerifyLinkedProgram(GLuint Program)
 {
 	SCOPE_CYCLE_COUNTER(STAT_OpenGLShaderLinkVerifyTime);
 
+#if UE_BUILD_DEBUG
 	GLint LinkStatus = 0;
 	glGetProgramiv(Program, GL_LINK_STATUS, &LinkStatus);
 	if (LinkStatus != GL_TRUE)
@@ -129,6 +133,7 @@ static bool VerifyLinkedProgram(GLuint Program)
 		}
 		return false;
 	}
+#endif
 	return true;
 }
 
@@ -329,7 +334,7 @@ ShaderType* CompileOpenGLShader(const TArray<uint8>& Code)
 	// The code as given to us.
 	FAnsiCharArray GlslCodeOriginal;
 	AppendCString(GlslCodeOriginal, (ANSICHAR*)Code.GetData() + CodeOffset);
-	uint32 GlslCodeOriginalCRC = FCrc::MemCrc_DEPRECATED(GlslCodeOriginal.GetData(), GlslCodeOriginal.Num() + 1);
+	uint32 GlslCodeOriginalCRC = FCrc::MemCrc_DEPRECATED(GlslCodeOriginal.GetData(), GlslCodeOriginal.Num());
 
 	// The amended code we actually compile.
 	FAnsiCharArray GlslCode;
@@ -367,7 +372,7 @@ ShaderType* CompileOpenGLShader(const TArray<uint8>& Code)
 		}
 #endif
 
-		Resource = glCreateShader(TypeEnum);
+		Resource = FOpenGL::CreateShader(TypeEnum);
 
 #if (PLATFORM_ANDROID || PLATFORM_HTML5)
 		if (IsES2Platform(GMaxRHIShaderPlatform))
@@ -1250,7 +1255,7 @@ static FOpenGLLinkedProgram* LinkProgram( const FOpenGLLinkedProgramConfiguratio
 	SCOPE_CYCLE_COUNTER(STAT_OpenGLShaderLinkTime);
 	VERIFY_GL_SCOPE();
 
-	GLuint Program = glCreateProgram();
+	GLuint Program = FOpenGL::CreateProgram();
 
 	// ensure that compute shaders are always alone
 	check( (Config.Shaders[CrossCompiler::SHADER_STAGE_VERTEX].Resource == 0) != (Config.Shaders[CrossCompiler::SHADER_STAGE_COMPUTE].Resource == 0));
@@ -1499,6 +1504,12 @@ FBoundShaderStateRHIRef FOpenGLDynamicRHI::RHICreateBoundShaderState(
 	{
 		check(VertexDeclarationRHI);
 
+		FShaderCache* ShaderCache = FShaderCache::GetShaderCache();
+		if(ShaderCache)
+		{
+			ShaderCache->LogBoundShaderState(FOpenGL::GetShaderPlatform(), VertexDeclarationRHI, VertexShaderRHI, PixelShaderRHI, HullShaderRHI, DomainShaderRHI, GeometryShaderRHI);
+		}
+
 		DYNAMIC_CAST_OPENGLRESOURCE(VertexDeclaration,VertexDeclaration);
 		DYNAMIC_CAST_OPENGLRESOURCE(VertexShader,VertexShader);
 		DYNAMIC_CAST_OPENGLRESOURCE(PixelShader,PixelShader);
diff --git a/Engine/Source/Runtime/OpenGLDrv/Private/OpenGLVertexDeclaration.cpp b/Engine/Source/Runtime/OpenGLDrv/Private/OpenGLVertexDeclaration.cpp
--- a/Engine/Source/Runtime/OpenGLDrv/Private/OpenGLVertexDeclaration.cpp
+++ b/Engine/Source/Runtime/OpenGLDrv/Private/OpenGLVertexDeclaration.cpp
@@ -5,6 +5,7 @@
 =============================================================================*/
 
 #include "OpenGLDrvPrivate.h"
+#include "ShaderCache.h"
 
 static FORCEINLINE void SetupGLElement(FOpenGLVertexElement& GLElement, GLenum Type, GLuint Size, bool bNormalized, bool bShouldConvertToFloat)
 {
@@ -133,6 +134,14 @@ FVertexDeclarationRHIRef FOpenGLDynamicRHI::RHICreateVertexDeclaration(const FVe
 	{
 		// Create and add to the cache if it doesn't exist.
 		VertexDeclarationRefPtr = &GOpenGLVertexDeclarationCache.Add(Key,new FOpenGLVertexDeclaration(Key.VertexElements));
+		
+		FShaderCache* ShaderCache = FShaderCache::GetShaderCache();
+		if(ShaderCache)
+		{
+			check(VertexDeclarationRefPtr);
+			check(IsValidRef(*VertexDeclarationRefPtr));
+			ShaderCache->LogVertexDeclaration(Elements, *VertexDeclarationRefPtr);
+		}
 	}
 
 	// The cached declaration must match the input declaration!
diff --git a/Engine/Source/Runtime/OpenGLDrv/Public/OpenGL.h b/Engine/Source/Runtime/OpenGLDrv/Public/OpenGL.h
--- a/Engine/Source/Runtime/OpenGLDrv/Public/OpenGL.h
+++ b/Engine/Source/Runtime/OpenGLDrv/Public/OpenGL.h
@@ -311,12 +311,7 @@ public:
 	static FPlatformOpenGLDevice*	CreateDevice() UGL_REQUIRED(NULL)
 	static FPlatformOpenGLContext*	CreateContext( FPlatformOpenGLDevice* Device, void* WindowHandle ) UGL_REQUIRED(NULL)
 
-	static FORCEINLINE void BufferSubData(GLenum Target, GLintptr Offset, GLsizeiptr Size, const GLvoid* Data)
-	{
-		glBufferSubData(Target, Offset, Size, Data);
-	}
-
-	static FORCEINLINE void CheckFrameBuffer() 
+	static FORCEINLINE void CheckFrameBuffer()
 	{
 #if UE_BUILD_DEBUG 
 		GLenum CompleteResult = glCheckFramebufferStatus(GL_FRAMEBUFFER);
@@ -327,21 +322,12 @@ public:
 #endif 
 	}
 
-	static FORCEINLINE void DeleteBuffers(GLsizei Number, const GLuint* Buffers)
-	{
-		glDeleteBuffers(Number, Buffers);
-	}
-
-	static FORCEINLINE void DeleteTextures(GLsizei Number, const GLuint* Textures)
-	{
-		glDeleteTextures(Number, Textures);
-	}
-
-	static FORCEINLINE void Flush()
-	{
-		glFlush();
-	}
-
+	static FORCEINLINE void BufferSubData(GLenum Target, GLintptr Offset, GLsizeiptr Size, const GLvoid* Data)	{ glBufferSubData(Target, Offset, Size, Data); }
+	static FORCEINLINE void DeleteBuffers(GLsizei Number, const GLuint* Buffers)	{ glDeleteBuffers(Number, Buffers); }
+	static FORCEINLINE void DeleteTextures(GLsizei Number, const GLuint* Textures)	{ glDeleteTextures(Number, Textures); }
+	static FORCEINLINE void Flush()													{ glFlush(); }
+	static FORCEINLINE GLuint CreateShader(GLenum Type)								{ return glCreateShader(Type); }
+	static FORCEINLINE GLuint CreateProgram()										{ return glCreateProgram(); }
 	static FORCEINLINE bool TimerQueryDisjoint()									{ return false; }
 
 protected:
diff --git a/Engine/Source/Runtime/OpenGLDrv/Public/OpenGL3.h b/Engine/Source/Runtime/OpenGLDrv/Public/OpenGL3.h
--- a/Engine/Source/Runtime/OpenGLDrv/Public/OpenGL3.h
+++ b/Engine/Source/Runtime/OpenGLDrv/Public/OpenGL3.h
@@ -415,6 +415,43 @@ struct FOpenGL3 : public FOpenGLBase
 	{
 		glCopyBufferSubData(ReadTarget, WriteTarget, ReadOffset, WriteOffset, Size);
 	}
+	
+	static FORCEINLINE GLuint CreateShader(GLenum Type)
+	{
+#if USE_OPENGL_NAME_CACHE
+		static TMap<GLenum, TArray<GLuint>> ShaderNames;
+		TArray<GLuint>& Shaders = ShaderNames.FindOrAdd(Type);
+		if(!Shaders.Num())
+		{
+			while(Shaders.Num() < OPENGL_NAME_CACHE_SIZE)
+			{
+				GLuint Resource = glCreateShader(Type);
+				Shaders.Add(Resource);
+			}
+		}
+		return Shaders.Pop();
+#else
+		return glCreateShader(Type);
+#endif
+	}
+	
+	static FORCEINLINE GLuint CreateProgram()
+	{
+#if USE_OPENGL_NAME_CACHE
+		static TArray<GLuint> ProgramNames;
+		if(!ProgramNames.Num())
+		{
+			while(ProgramNames.Num() < OPENGL_NAME_CACHE_SIZE)
+			{
+				GLuint Resource = glCreateProgram();
+				ProgramNames.Add(Resource);
+			}
+		}
+		return ProgramNames.Pop();
+#else
+		return glCreateProgram();
+#endif
+	}
 
 	static FORCEINLINE void GenBuffers( GLsizei n, GLuint *buffers)
 	{
diff --git a/Engine/Source/Runtime/RHI/Public/RHIResources.h b/Engine/Source/Runtime/RHI/Public/RHIResources.h
--- a/Engine/Source/Runtime/RHI/Public/RHIResources.h
+++ b/Engine/Source/Runtime/RHI/Public/RHIResources.h
@@ -7,6 +7,7 @@
 #include "RefCounting.h"
 #include "Runtime/Engine/Public/PixelFormat.h" // for EPixelFormat
 #include "LockFreeList.h"
+#include "SecureHash.h"
 
 #define DISABLE_RHI_DEFFERED_DELETE 0
 
@@ -103,12 +104,24 @@ class FRHIBoundShaderState : public FRHIResource {};
 // Shaders
 //
 
-class FRHIVertexShader : public FRHIResource {};
-class FRHIHullShader : public FRHIResource {};
-class FRHIDomainShader : public FRHIResource {};
-class FRHIPixelShader : public FRHIResource {};
-class FRHIGeometryShader : public FRHIResource {};
-class FRHIComputeShader : public FRHIResource {};
+class FRHIShader : public FRHIResource
+{
+public:
+	FRHIShader(bool InbDoNotDeferDelete = false) : FRHIResource(InbDoNotDeferDelete) {}
+	
+	void SetHash(FSHAHash InHash) { Hash = InHash; }
+	FSHAHash GetHash() const { return Hash; }
+
+private:
+	FSHAHash Hash;
+};
+
+class FRHIVertexShader : public FRHIShader {};
+class FRHIHullShader : public FRHIShader {};
+class FRHIDomainShader : public FRHIShader {};
+class FRHIPixelShader : public FRHIShader {};
+class FRHIGeometryShader : public FRHIShader {};
+class FRHIComputeShader : public FRHIShader {};
 
 //
 // Buffers
diff --git a/Engine/Source/Runtime/ShaderCore/Private/Shader.cpp b/Engine/Source/Runtime/ShaderCore/Private/Shader.cpp
--- a/Engine/Source/Runtime/ShaderCore/Private/Shader.cpp
+++ b/Engine/Source/Runtime/ShaderCore/Private/Shader.cpp
@@ -12,6 +12,7 @@
 #include "ModuleManager.h"
 #include "TargetPlatform.h"
 #include "RHICommandList.h"
+#include "ShaderCache.h"
 
 
 DEFINE_LOG_CATEGORY(LogShaders);
@@ -338,6 +339,12 @@ void FShaderResource::Serialize(FArchive& Ar)
 	{
 		INC_DWORD_STAT_BY_FName(GetMemoryStatType((EShaderFrequency)Target.Frequency).GetName(), (int64)Code.Num());
 		INC_DWORD_STAT_BY(STAT_Shaders_ShaderResourceMemory, GetSizeBytes());
+		
+		FShaderCache* ShaderCache = FShaderCache::GetShaderCache();
+		if(ShaderCache)
+		{
+			ShaderCache->LogShader((EShaderPlatform)Target.Platform, (EShaderFrequency)Target.Frequency, OutputHash, Code);
+		}
 	}
 }
 
@@ -443,29 +450,31 @@ void FShaderResource::InitRHI()
 	INC_DWORD_STAT_BY(STAT_Shaders_NumShadersUsedForRendering, 1);
 	SCOPE_CYCLE_COUNTER(STAT_Shaders_RTShaderLoadTime);
 
+	FShaderCache* ShaderCache = FShaderCache::GetShaderCache();
+
 	if(Target.Frequency == SF_Vertex)
 	{
-		VertexShader = RHICreateVertexShader(Code);
+		VertexShader = ShaderCache ? ShaderCache->GetVertexShader((EShaderPlatform)Target.Platform, OutputHash, Code) : RHICreateVertexShader(Code);
 	}
 	else if(Target.Frequency == SF_Pixel)
 	{
-		PixelShader = RHICreatePixelShader(Code);
+		PixelShader = ShaderCache ? ShaderCache->GetPixelShader((EShaderPlatform)Target.Platform, OutputHash, Code) : RHICreatePixelShader(Code);
 	}
 	else if(Target.Frequency == SF_Hull)
 	{
-		HullShader = RHICreateHullShader(Code);
+		HullShader = ShaderCache ? ShaderCache->GetHullShader((EShaderPlatform)Target.Platform, OutputHash, Code) : RHICreateHullShader(Code);
 	}
 	else if(Target.Frequency == SF_Domain)
 	{
-		DomainShader = RHICreateDomainShader(Code);
+		DomainShader = ShaderCache ? ShaderCache->GetDomainShader((EShaderPlatform)Target.Platform, OutputHash, Code) : RHICreateDomainShader(Code);
 	}
 	else if(Target.Frequency == SF_Geometry)
 	{
-		GeometryShader = RHICreateGeometryShader(Code);
+		GeometryShader = ShaderCache ? ShaderCache->GetGeometryShader((EShaderPlatform)Target.Platform, OutputHash, Code) : RHICreateGeometryShader(Code);
 	}
 	else if(Target.Frequency == SF_Compute)
 	{
-		ComputeShader = RHICreateComputeShader(Code);
+		ComputeShader = ShaderCache ? ShaderCache->GetComputeShader((EShaderPlatform)Target.Platform, OutputHash, Code) : RHICreateComputeShader(Code);
 	}
 
 	if (!FPlatformProperties::HasEditorOnlyData())
diff --git a/Engine/Source/Runtime/ShaderCore/Private/ShaderCache.cpp b/Engine/Source/Runtime/ShaderCore/Private/ShaderCache.cpp
new file mode 100644
--- /dev/null
+++ b/Engine/Source/Runtime/ShaderCore/Private/ShaderCache.cpp
@@ -0,0 +1,441 @@
+// Copyright 1998-2015 Epic Games, Inc. All Rights Reserved.
+
+/*=============================================================================
+	ShaderCache.cpp: Bound shader state cache implementation.
+=============================================================================*/
+
+#include "ShaderCore.h"
+#include "ShaderCache.h"
+#include "RHI.h"
+#include "RenderingThread.h"
+
+int32 bUseShaderCaching = PLATFORM_MAC ? 1 : 0;
+static FAutoConsoleVariableRef CVarUseShaderCaching(
+	TEXT("r.UseShaderCaching"),
+	bUseShaderCaching,
+	TEXT("If true, log all shaders & bound-shader-states, so that they may be instantiated in the RHI on deserialisation rather than waiting for first use.")
+	);
+
+FShaderCache* FShaderCache::Cache = nullptr;
+
+void FShaderCache::InitShaderCache()
+{
+	check(!Cache);
+	if(bUseShaderCaching)
+	{
+		Cache = new FShaderCache;
+	}
+}
+
+FShaderCache* FShaderCache::GetShaderCache()
+{
+	return bUseShaderCaching ? Cache : nullptr;
+}
+
+void FShaderCache::ShutdownShaderCache()
+{
+	check(Cache);
+	delete Cache;
+	Cache = nullptr;
+}
+
+FShaderCache::FShaderCache()
+{
+	FString UserBinaryShaderFile = FPaths::GameSavedDir() / TEXT("ShaderCache.ushaderprecache");
+	FString GameBinaryShaderFile = FPaths::GameContentDir() / TEXT("ShaderCache.ushaderprecache");
+
+	FArchive* BinaryShaderAr = nullptr;
+	if ( IFileManager::Get().FileSize(*UserBinaryShaderFile) > 0 )
+	{
+		BinaryShaderAr = IFileManager::Get().CreateFileReader(*UserBinaryShaderFile);
+	}
+	else
+	{
+		BinaryShaderAr = IFileManager::Get().CreateFileReader(*GameBinaryShaderFile);
+	}
+
+	if ( BinaryShaderAr != nullptr )
+	{
+		*BinaryShaderAr << *this;
+		delete BinaryShaderAr;
+	}
+}
+
+FShaderCache::~FShaderCache()
+{
+	FString BinaryShaderFile = FPaths::GameSavedDir() / TEXT("ShaderCache.ushaderprecache");
+	FArchive* BinaryShaderAr = IFileManager::Get().CreateFileWriter(*BinaryShaderFile);
+	if( BinaryShaderAr != NULL )
+	{
+		*BinaryShaderAr << *this;
+		delete BinaryShaderAr;
+	}
+}
+
+FVertexShaderRHIRef FShaderCache::GetVertexShader(EShaderPlatform Platform, FSHAHash Hash, TArray<uint8> const& Code)
+{
+	FShaderCacheKey Key;
+	Key.Platform = Platform;
+	Key.Frequency = SF_Vertex;
+	Key.Hash = Hash;
+	Key.bActive = true;
+	FVertexShaderRHIRef Shader = CachedVertexShaders.FindRef(Key);
+	if(!IsValidRef(Shader))
+	{
+		FShaderPlatformCache& PlatformCache = Caches.FindOrAdd(Platform);
+		PlatformCache.Shaders.Add(Key);
+		Shader = RHICreateVertexShader(Code);
+		check(IsValidRef(Shader));
+		Shader->SetHash(Hash);
+		CachedVertexShaders.Add(Key, Shader);
+		PrebindShader(Key);
+	}
+	return Shader;
+}
+
+FPixelShaderRHIRef FShaderCache::GetPixelShader(EShaderPlatform Platform, FSHAHash Hash, TArray<uint8> const& Code)
+{
+	FShaderCacheKey Key;
+	Key.Platform = Platform;
+	Key.Frequency = SF_Pixel;
+	Key.Hash = Hash;
+	Key.bActive = true;
+	FPixelShaderRHIRef Shader = CachedPixelShaders.FindRef(Key);
+	if(!IsValidRef(Shader))
+	{
+		FShaderPlatformCache& PlatformCache = Caches.FindOrAdd(Platform);
+		PlatformCache.Shaders.Add(Key);
+		Shader = RHICreatePixelShader(Code);
+		check(IsValidRef(Shader));
+		Shader->SetHash(Hash);
+		CachedPixelShaders.Add(Key, Shader);
+		PrebindShader(Key);
+	}
+	return Shader;
+}
+
+FGeometryShaderRHIRef FShaderCache::GetGeometryShader(EShaderPlatform Platform, FSHAHash Hash, TArray<uint8> const& Code)
+{
+	FShaderCacheKey Key;
+	Key.Platform = Platform;
+	Key.Frequency = SF_Geometry;
+	Key.Hash = Hash;
+	Key.bActive = true;
+	FGeometryShaderRHIRef Shader = CachedGeometryShaders.FindRef(Key);
+	if(!IsValidRef(Shader))
+	{
+		FShaderPlatformCache& PlatformCache = Caches.FindOrAdd(Platform);
+		PlatformCache.Shaders.Add(Key);
+		Shader = RHICreateGeometryShader(Code);
+		check(IsValidRef(Shader));
+		Shader->SetHash(Hash);
+		CachedGeometryShaders.Add(Key, Shader);
+		PrebindShader(Key);
+	}
+	return Shader;
+}
+
+FHullShaderRHIRef FShaderCache::GetHullShader(EShaderPlatform Platform, FSHAHash Hash, TArray<uint8> const& Code)
+{
+	FShaderCacheKey Key;
+	Key.Platform = Platform;
+	Key.Frequency = SF_Hull;
+	Key.Hash = Hash;
+	Key.bActive = true;
+	FHullShaderRHIRef Shader = CachedHullShaders.FindRef(Key);
+	if(!IsValidRef(Shader))
+	{
+		FShaderPlatformCache& PlatformCache = Caches.FindOrAdd(Platform);
+		PlatformCache.Shaders.Add(Key);
+		Shader = RHICreateHullShader(Code);
+		check(IsValidRef(Shader));
+		Shader->SetHash(Hash);
+		CachedHullShaders.Add(Key, Shader);
+		PrebindShader(Key);
+	}
+	return Shader;
+}
+
+FDomainShaderRHIRef FShaderCache::GetDomainShader(EShaderPlatform Platform, FSHAHash Hash, TArray<uint8> const& Code)
+{
+	FShaderCacheKey Key;
+	Key.Platform = Platform;
+	Key.Frequency = SF_Domain;
+	Key.Hash = Hash;
+	Key.bActive = true;
+	FDomainShaderRHIRef Shader = CachedDomainShaders.FindRef(Key);
+	if(!IsValidRef(Shader))
+	{
+		FShaderPlatformCache& PlatformCache = Caches.FindOrAdd(Platform);
+		PlatformCache.Shaders.Add(Key);
+		Shader = RHICreateDomainShader(Code);
+		check(IsValidRef(Shader));
+		Shader->SetHash(Hash);
+		CachedDomainShaders.Add(Key, Shader);
+		PrebindShader(Key);
+	}
+	return Shader;
+}
+
+FComputeShaderRHIRef FShaderCache::GetComputeShader(EShaderPlatform Platform, FSHAHash Hash, TArray<uint8> const& Code)
+{
+	FShaderCacheKey Key;
+	Key.Platform = Platform;
+	Key.Frequency = SF_Compute;
+	Key.Hash = Hash;
+	Key.bActive = true;
+	FComputeShaderRHIRef Shader = CachedComputeShaders.FindRef(Key);
+	if(!IsValidRef(Shader))
+	{
+		FShaderPlatformCache& PlatformCache = Caches.FindOrAdd(Platform);
+		PlatformCache.Shaders.Add(Key);
+		Shader = RHICreateComputeShader(Code);
+		check(IsValidRef(Shader));
+		Shader->SetHash(Hash);
+		CachedComputeShaders.Add(Key, Shader);
+		PrebindShader(Key);
+	}
+	return Shader;
+}
+
+void FShaderCache::LogShader(EShaderPlatform Platform, EShaderFrequency Frequency, FSHAHash Hash, TArray<uint8> const& Code)
+{
+	FShaderCacheKey Key;
+	Key.Hash = Hash;
+	Key.Platform = Platform;
+	Key.Frequency = Frequency;
+	Key.bActive = true;
+	ENQUEUE_UNIQUE_RENDER_COMMAND_TWOPARAMETER(LogShader, FShaderCacheKey,Key,Key, TArray<uint8>,Code,Code, {
+		FShaderCache::GetShaderCache()->SubmitShader(Key, Code);
+	});
+}
+
+void FShaderCache::LogVertexDeclaration(const FVertexDeclarationElementList& VertexElements, FVertexDeclarationRHIParamRef VertexDeclaration)
+{
+	VertexDeclarations.Add(VertexDeclaration, VertexElements);
+}
+
+void FShaderCache::LogBoundShaderState(EShaderPlatform Platform, FVertexDeclarationRHIParamRef VertexDeclaration,
+								FVertexShaderRHIParamRef VertexShader,
+								FPixelShaderRHIParamRef PixelShader,
+								FHullShaderRHIParamRef HullShader,
+								FDomainShaderRHIParamRef DomainShader,
+								FGeometryShaderRHIParamRef GeometryShader)
+{
+	FShaderCacheBoundState Info;
+	if(VertexDeclaration)
+	{
+		Info.VertexDeclaration = VertexDeclarations.FindChecked(VertexDeclaration);
+	}
+	if(VertexShader)
+	{
+		Info.VertexShader.Platform = Platform;
+		Info.VertexShader.Frequency = SF_Vertex;
+		Info.VertexShader.Hash = VertexShader->GetHash();
+		Info.VertexShader.bActive = true;
+	}
+	if(PixelShader)
+	{
+		Info.PixelShader.Platform = Platform;
+		Info.PixelShader.Frequency = SF_Pixel;
+		Info.PixelShader.Hash = PixelShader->GetHash();
+		Info.PixelShader.bActive = true;
+	}
+	if(GeometryShader)
+	{
+		Info.GeometryShader.Platform = Platform;
+		Info.GeometryShader.Frequency = SF_Geometry;
+		Info.GeometryShader.Hash = GeometryShader->GetHash();
+		Info.GeometryShader.bActive = true;
+	}
+	if(HullShader)
+	{
+		Info.HullShader.Platform = Platform;
+		Info.HullShader.Frequency = SF_Hull;
+		Info.HullShader.Hash = HullShader->GetHash();
+		Info.HullShader.bActive = true;
+	}
+	if(DomainShader)
+	{
+		Info.DomainShader.Platform = Platform;
+		Info.DomainShader.Frequency = SF_Domain;
+		Info.DomainShader.Hash = DomainShader->GetHash();
+		Info.DomainShader.bActive = true;
+	}
+	
+	FShaderPlatformCache& PlatformCache = Caches.FindOrAdd(Platform);
+	PlatformCache.BoundShaderStates.Add(Info);
+	
+	if(VertexShader)
+	{
+		TSet<FShaderCacheBoundState>& Set = PlatformCache.ShaderStateMembership.FindOrAdd(Info.VertexShader);
+		if(!Set.Find(Info))
+		{
+			Set.Add(Info);
+		}
+	}
+	if(PixelShader)
+	{
+		TSet<FShaderCacheBoundState>& Set = PlatformCache.ShaderStateMembership.FindOrAdd(Info.PixelShader);
+		if(!Set.Find(Info))
+		{
+			Set.Add(Info);
+		}
+	}
+	if(GeometryShader)
+	{
+		TSet<FShaderCacheBoundState>& Set = PlatformCache.ShaderStateMembership.FindOrAdd(Info.GeometryShader);
+		if(!Set.Find(Info))
+		{
+			Set.Add(Info);
+		}
+	}
+	if(HullShader)
+	{
+		TSet<FShaderCacheBoundState>& Set = PlatformCache.ShaderStateMembership.FindOrAdd(Info.HullShader);
+		if(!Set.Find(Info))
+		{
+			Set.Add(Info);
+		}
+	}
+	if(DomainShader)
+	{
+		TSet<FShaderCacheBoundState>& Set = PlatformCache.ShaderStateMembership.FindOrAdd(Info.DomainShader);
+		if(!Set.Find(Info))
+		{
+			Set.Add(Info);
+		}
+	}
+}
+
+void FShaderCache::PrebindShader(FShaderCacheKey Key)
+{
+	FShaderPlatformCache& PlatformCache = Caches.FindOrAdd(Key.Platform);
+	TSet<FShaderCacheBoundState>& BoundStates = PlatformCache.ShaderStateMembership.FindOrAdd(Key);
+	for (FShaderCacheBoundState State : BoundStates)
+	{
+		FVertexShaderRHIRef VertexShader = State.VertexShader.bActive ? CachedVertexShaders.FindRef(State.VertexShader) : nullptr;
+		FPixelShaderRHIRef PixelShader = State.PixelShader.bActive ? CachedPixelShaders.FindRef(State.PixelShader) : nullptr;
+		FGeometryShaderRHIRef GeometryShader = State.GeometryShader.bActive ? CachedGeometryShaders.FindRef(State.GeometryShader) : nullptr;
+		FHullShaderRHIRef HullShader = State.HullShader.bActive ? CachedHullShaders.FindRef(State.HullShader) : nullptr;
+		FDomainShaderRHIRef DomainShader = State.DomainShader.bActive ? CachedDomainShaders.FindRef(State.DomainShader) : nullptr;
+		
+		bool bOK = true;
+		bOK &= (State.VertexShader.bActive == IsValidRef(VertexShader));
+		bOK &= (State.PixelShader.bActive == IsValidRef(PixelShader));
+		bOK &= (State.GeometryShader.bActive == IsValidRef(GeometryShader));
+		bOK &= (State.HullShader.bActive == IsValidRef(HullShader));
+		bOK &= (State.DomainShader.bActive == IsValidRef(DomainShader));
+		
+		if(bOK)
+		{
+			FVertexDeclarationRHIRef VertexDeclaration = RHICreateVertexDeclaration(State.VertexDeclaration);
+			bOK &= IsValidRef(VertexDeclaration);
+			
+			if(bOK)
+			{
+				FBoundShaderStateRHIRef BoundState = RHICreateBoundShaderState(
+					VertexDeclaration,
+					VertexShader,
+					HullShader,
+					DomainShader,
+					PixelShader,
+					GeometryShader
+				);
+				if(IsValidRef(BoundState))
+				{
+					BoundShaderStates.Add(State, BoundState);
+				}
+			}
+		}
+	}
+}
+
+void FShaderCache::SubmitShader(FShaderCacheKey Key, TArray<uint8> const& Code)
+{
+	FShaderPlatformCache& PlatformCache = Caches.FindOrAdd(Key.Platform);
+	switch (Key.Frequency)
+	{
+		case SF_Vertex:
+			if(!CachedVertexShaders.Find(Key))
+			{
+				FVertexShaderRHIRef Shader = RHICreateVertexShader(Code);
+				if(IsValidRef(Shader))
+				{
+					Shader->SetHash(Key.Hash);
+					CachedVertexShaders.Add(Key, Shader);
+					PrebindShader(Key);
+					PlatformCache.Shaders.Add(Key);
+				}
+			}
+			break;
+		case SF_Pixel:
+			if(!CachedPixelShaders.Find(Key))
+			{
+				FPixelShaderRHIRef Shader = RHICreatePixelShader(Code);
+				if(IsValidRef(Shader))
+				{
+					Shader->SetHash(Key.Hash);
+					CachedPixelShaders.Add(Key, Shader);
+					PrebindShader(Key);
+					PlatformCache.Shaders.Add(Key);
+				}
+			}
+			break;
+		case SF_Geometry:
+			if(!CachedGeometryShaders.Find(Key))
+			{
+				FGeometryShaderRHIRef Shader = RHICreateGeometryShader(Code);
+				if(IsValidRef(Shader))
+				{
+					Shader->SetHash(Key.Hash);
+					CachedGeometryShaders.Add(Key, Shader);
+					PrebindShader(Key);
+					PlatformCache.Shaders.Add(Key);
+				}
+			}
+			break;
+		case SF_Hull:
+			if(!CachedHullShaders.Find(Key))
+			{
+				FHullShaderRHIRef Shader = RHICreateHullShader(Code);
+				if(IsValidRef(Shader))
+				{
+					Shader->SetHash(Key.Hash);
+					CachedHullShaders.Add(Key, Shader);
+					PrebindShader(Key);
+					PlatformCache.Shaders.Add(Key);
+				}
+			}
+			break;
+		case SF_Domain:
+			if(!CachedDomainShaders.Find(Key))
+			{
+				FDomainShaderRHIRef Shader = RHICreateDomainShader(Code);
+				if(IsValidRef(Shader))
+				{
+					Shader->SetHash(Key.Hash);
+					CachedDomainShaders.Add(Key, Shader);
+					PrebindShader(Key);
+					PlatformCache.Shaders.Add(Key);
+				}
+			}
+			break;
+		case SF_Compute:
+			if(!CachedComputeShaders.Find(Key))
+			{
+				FComputeShaderRHIRef Shader = RHICreateComputeShader(Code);
+				if(IsValidRef(Shader))
+				{
+					Shader->SetHash(Key.Hash);
+					CachedComputeShaders.Add(Key, Shader);
+					PrebindShader(Key);
+					PlatformCache.Shaders.Add(Key);
+				}
+			}
+			break;
+		default:
+			check(false);
+			break;
+	}
+}
diff --git a/Engine/Source/Runtime/ShaderCore/Public/ShaderCache.h b/Engine/Source/Runtime/ShaderCore/Public/ShaderCache.h
new file mode 100644
--- /dev/null
+++ b/Engine/Source/Runtime/ShaderCore/Public/ShaderCache.h
@@ -0,0 +1,155 @@
+// Copyright 1998-2015 Epic Games, Inc. All Rights Reserved.
+
+/*=============================================================================
+	ShaderCache.h: Shader precompilation mechanism
+=============================================================================*/
+
+#pragma once
+
+#include "ShaderCore.h"
+#include "RHI.h"
+#include "BoundShaderStateCache.h"
+
+class SHADERCORE_API FShaderCache
+{
+	struct SHADERCORE_API FShaderCacheKey
+	{
+		FShaderCacheKey() : Platform(SP_NumPlatforms), Frequency(SF_NumFrequencies), bActive(false) {}
+	
+		FSHAHash Hash;
+		EShaderPlatform Platform;
+		EShaderFrequency Frequency;
+		bool bActive;
+		
+		friend bool operator ==(const FShaderCacheKey& A,const FShaderCacheKey& B)
+		{
+			return A.Hash == B.Hash && A.Platform == B.Platform && A.Frequency == B.Frequency && A.bActive == B.bActive;
+		}
+
+		friend uint32 GetTypeHash(const FShaderCacheKey &Key)
+		{
+			uint32 TargetFrequency = Key.Frequency;
+			uint32 TargetPlatform = Key.Platform;
+			return FCrc::MemCrc_DEPRECATED((const void*)&Key.Hash, sizeof(Key.Hash)) ^ GetTypeHash(TargetPlatform) ^ GetTypeHash(TargetFrequency) ^ GetTypeHash(Key.bActive);
+		}
+		
+		friend FArchive& operator<<( FArchive& Ar, FShaderCacheKey& Info )
+		{
+			uint32 TargetFrequency = Info.Frequency;
+			uint32 TargetPlatform = Info.Platform;
+			Ar << TargetFrequency << TargetPlatform;
+			Info.Frequency = (EShaderFrequency)TargetFrequency;
+			Info.Platform = (EShaderPlatform)TargetPlatform;
+			return Ar << Info.Hash << Info.bActive;
+		}
+	};
+		
+public:
+	FShaderCache();
+	~FShaderCache();
+
+	static void InitShaderCache();
+	static FShaderCache* GetShaderCache();
+	static void ShutdownShaderCache();
+
+	FVertexShaderRHIRef GetVertexShader(EShaderPlatform Platform, FSHAHash Hash, TArray<uint8> const& Code);
+	FPixelShaderRHIRef GetPixelShader(EShaderPlatform Platform, FSHAHash Hash, TArray<uint8> const& Code);
+	FGeometryShaderRHIRef GetGeometryShader(EShaderPlatform Platform, FSHAHash Hash, TArray<uint8> const& Code);
+	FHullShaderRHIRef GetHullShader(EShaderPlatform Platform, FSHAHash Hash, TArray<uint8> const& Code);
+	FDomainShaderRHIRef GetDomainShader(EShaderPlatform Platform, FSHAHash Hash, TArray<uint8> const& Code);
+	FComputeShaderRHIRef GetComputeShader(EShaderPlatform Platform, FSHAHash Hash, TArray<uint8> const& Code);
+	FVertexDeclarationRHIParamRef GetVertexDeclaration(FVertexDeclarationElementList& VertexElements);
+	
+	void LogShader(EShaderPlatform Platform, EShaderFrequency Frequency, FSHAHash Hash, TArray<uint8> const& Code);
+	
+	void LogVertexDeclaration(const FVertexDeclarationElementList& VertexElements, FVertexDeclarationRHIParamRef VertexDeclaration);
+	
+	void LogBoundShaderState(EShaderPlatform Platform, FVertexDeclarationRHIParamRef VertexDeclaration,
+								FVertexShaderRHIParamRef VertexShader,
+								FPixelShaderRHIParamRef PixelShader,
+								FHullShaderRHIParamRef HullShader,
+								FDomainShaderRHIParamRef DomainShader,
+								FGeometryShaderRHIParamRef GeometryShader);
+		
+	friend FArchive& operator<<( FArchive& Ar, FShaderCache& Info )
+	{
+		return Ar << Info.Caches;
+	}
+	
+private:
+	void PrebindShader(FShaderCacheKey Key);
+	void SubmitShader(FShaderCacheKey Key, TArray<uint8> const& Code);
+	
+private:
+	struct FShaderCacheBoundState
+	{
+		FVertexDeclarationElementList VertexDeclaration;
+		FShaderCacheKey VertexShader;
+		FShaderCacheKey PixelShader;
+		FShaderCacheKey GeometryShader;
+		FShaderCacheKey HullShader;
+		FShaderCacheKey DomainShader;
+		
+		friend bool operator ==(const FShaderCacheBoundState& A,const FShaderCacheBoundState& B)
+		{
+			if(A.VertexDeclaration.Num() == B.VertexDeclaration.Num())
+			{
+				for(uint32 i = 0; i < A.VertexDeclaration.Num(); i++)
+				{
+					if(FMemory::Memcmp(&A.VertexDeclaration[i], &B.VertexDeclaration[i], sizeof(FVertexElement)))
+					{
+						return false;
+					}
+				}
+				return A.VertexShader == B.VertexShader && A.PixelShader == B.PixelShader && A.GeometryShader == B.GeometryShader && A.HullShader == B.HullShader && A.DomainShader == B.DomainShader;
+			}
+			return false;
+		}
+
+		friend uint32 GetTypeHash(const FShaderCacheBoundState &Key)
+		{
+			uint32 Hash = 0;
+			
+			for(auto Element : Key.VertexDeclaration)
+			{
+				Hash ^= FCrc::MemCrc_DEPRECATED(&Element, sizeof(FVertexElement));
+			}
+			
+			return Hash ^ GetTypeHash(Key.VertexShader) ^ GetTypeHash(Key.PixelShader) ^ GetTypeHash(Key.GeometryShader) ^ GetTypeHash(Key.HullShader) ^ GetTypeHash(Key.DomainShader);
+		}
+		
+		friend FArchive& operator<<( FArchive& Ar, FShaderCacheBoundState& Info )
+		{
+			return Ar << Info.VertexDeclaration << Info.VertexShader << Info.PixelShader << Info.GeometryShader << Info.HullShader << Info.DomainShader;
+		}
+	};
+
+	struct FShaderPlatformCache
+	{
+		TSet<FShaderCacheKey> Shaders;
+		TSet<FShaderCacheBoundState> BoundShaderStates;
+		TMap<FShaderCacheKey, TSet<FShaderCacheBoundState>> ShaderStateMembership;
+		
+		friend FArchive& operator<<( FArchive& Ar, FShaderPlatformCache& Info )
+		{
+			return Ar << Info.Shaders << Info.BoundShaderStates << Info.ShaderStateMembership;
+		}
+	};
+
+	typedef TMap<uint32, FShaderPlatformCache> FShaderCaches;
+
+	// Serialised
+	FShaderCaches Caches;
+	
+	// Transient
+	TMap<FShaderCacheKey, FVertexShaderRHIRef> CachedVertexShaders;
+	TMap<FShaderCacheKey, FPixelShaderRHIRef> CachedPixelShaders;
+	TMap<FShaderCacheKey, FGeometryShaderRHIRef> CachedGeometryShaders;
+	TMap<FShaderCacheKey, FHullShaderRHIRef> CachedHullShaders;
+	TMap<FShaderCacheKey, FDomainShaderRHIRef> CachedDomainShaders;
+	TMap<FShaderCacheKey, FComputeShaderRHIRef> CachedComputeShaders;
+	TMap<FVertexDeclarationRHIParamRef, FVertexDeclarationElementList> VertexDeclarations;
+	TMap<FShaderCacheBoundState, FBoundShaderStateRHIRef> BoundShaderStates;
+	
+	static FShaderCache* Cache;
+};
