[INTEGRATE] Merged CL #2424059 from UE4-UT to Main

Fix signed/unsigned mismatch in ShaderCache.h.

diff --git a/Engine/Source/Runtime/ShaderCore/Public/ShaderCache.h b/Engine/Source/Runtime/ShaderCore/Public/ShaderCache.h
--- a/Engine/Source/Runtime/ShaderCore/Public/ShaderCache.h
+++ b/Engine/Source/Runtime/ShaderCore/Public/ShaderCache.h
@@ -94,7 +94,7 @@ private:
 		{
 			if(A.VertexDeclaration.Num() == B.VertexDeclaration.Num())
 			{
-				for(uint32 i = 0; i < A.VertexDeclaration.Num(); i++)
+				for(int32 i = 0; i < A.VertexDeclaration.Num(); i++)
 				{
 					if(FMemory::Memcmp(&A.VertexDeclaration[i], &B.VertexDeclaration[i], sizeof(FVertexElement)))
 					{
