[INTEGRATE] Merged CL #2425500 from UE4-UT to Main

Fix crash on shutdown when ShaderCache has not been initialised.

diff --git a/Engine/Source/Runtime/ShaderCore/Private/ShaderCache.cpp b/Engine/Source/Runtime/ShaderCore/Private/ShaderCache.cpp
--- a/Engine/Source/Runtime/ShaderCore/Private/ShaderCache.cpp
+++ b/Engine/Source/Runtime/ShaderCore/Private/ShaderCache.cpp
@@ -34,9 +34,11 @@ FShaderCache* FShaderCache::GetShaderCache()
 
 void FShaderCache::ShutdownShaderCache()
 {
-	check(Cache);
-	delete Cache;
-	Cache = nullptr;
+	if (Cache)
+	{
+		delete Cache;
+		Cache = nullptr;
+	}
 }
 
 FShaderCache::FShaderCache()
