Reworked shader stream out support, now works properly with FShaderResource sharing
* GetStreamOutElements is a callback function on the shadertype, just like ModifyCompilationEnvironment
* If GetStreamOutElements is specified, the FShaderResource includes that shader type in its ID as it can't be shared with any other shader type.
* FShader is no longer a FRenderResource which solves some multiple initialization issues

diff --git a/Engine/Shaders/ConvertToUniformMesh.usf b/Engine/Shaders/ConvertToUniformMesh.usf
--- a/Engine/Shaders/ConvertToUniformMesh.usf
+++ b/Engine/Shaders/ConvertToUniformMesh.usf
@@ -29,7 +29,7 @@ void ConvertToUniformMeshVS(
 	Output.PixelPosition = WorldPosition;
 }
 
-// Must match GUniformMeshVertexLayout and TRIANGLE_VERTEX_DATA_STRIDE
+// Must match GetUniformMeshStreamOutLayout and TRIANGLE_VERTEX_DATA_STRIDE
 struct FConvertToUniformMeshVertex
 {
 	float4 Position : SV_Position;
diff --git a/Engine/Source/Runtime/Engine/Private/GlobalShader.cpp b/Engine/Source/Runtime/Engine/Private/GlobalShader.cpp
--- a/Engine/Source/Runtime/Engine/Private/GlobalShader.cpp
+++ b/Engine/Source/Runtime/Engine/Private/GlobalShader.cpp
@@ -112,9 +112,11 @@ FShader* FGlobalShaderType::FinishCompileShader(const FShaderCompileJob& Current
 {
 	if (CurrentJob.bSucceeded)
 	{
+		FShaderType* SpecificType = CurrentJob.ShaderType->LimitShaderResourceToThisType() ? CurrentJob.ShaderType : NULL;
+
 		// Reuse an existing resource with the same key or create a new one based on the compile output
 		// This allows FShaders to share compiled bytecode and RHI shader references
-		FShaderResource* Resource = FShaderResource::FindOrCreateShaderResource(CurrentJob.Output);
+		FShaderResource* Resource = FShaderResource::FindOrCreateShaderResource(CurrentJob.Output, SpecificType);
 		check(Resource);
 
 		// Find a shader with the same key in memory
@@ -429,7 +431,7 @@ TShaderMap<FGlobalShaderType>* GetGlobalShaderMap(EShaderPlatform Platform, bool
 
 				if (Shader)
 				{
-					Shader->InitializeResource();
+					Shader->BeginInitializeResources();
 				}
 			}
 		}
diff --git a/Engine/Source/Runtime/Engine/Private/Materials/MaterialShader.cpp b/Engine/Source/Runtime/Engine/Private/Materials/MaterialShader.cpp
--- a/Engine/Source/Runtime/Engine/Private/Materials/MaterialShader.cpp
+++ b/Engine/Source/Runtime/Engine/Private/Materials/MaterialShader.cpp
@@ -658,9 +658,11 @@ FShader* FMaterialShaderType::FinishCompileShader(
 {
 	check(CurrentJob.bSucceeded);
 
+	FShaderType* SpecificType = CurrentJob.ShaderType->LimitShaderResourceToThisType() ? CurrentJob.ShaderType : NULL;
+
 	// Reuse an existing resource with the same key or create a new one based on the compile output
 	// This allows FShaders to share compiled bytecode and RHI shader references
-	FShaderResource* Resource = FShaderResource::FindOrCreateShaderResource(CurrentJob.Output);
+	FShaderResource* Resource = FShaderResource::FindOrCreateShaderResource(CurrentJob.Output, SpecificType);
 
 	// Find a shader with the same key in memory
 	FShader* Shader = CurrentJob.ShaderType->FindShaderById(FShaderId(MaterialShaderMapHash, CurrentJob.VFType, CurrentJob.ShaderType, CurrentJob.Input.Target));
@@ -1467,7 +1469,7 @@ void FMaterialShaderMap::Register(EShaderPlatform InShaderPlatform)
 
 			if (Shader)
 			{
-				Shader->InitializeResource();
+				Shader->BeginInitializeResources();
 			}
 		}
 
@@ -1479,7 +1481,7 @@ void FMaterialShaderMap::Register(EShaderPlatform InShaderPlatform)
 
 				if (Shader)
 				{
-					Shader->InitializeResource();
+					Shader->BeginInitializeResources();
 				}
 			}
 		}
diff --git a/Engine/Source/Runtime/Engine/Private/Materials/MeshMaterialShader.cpp b/Engine/Source/Runtime/Engine/Private/Materials/MeshMaterialShader.cpp
--- a/Engine/Source/Runtime/Engine/Private/Materials/MeshMaterialShader.cpp
+++ b/Engine/Source/Runtime/Engine/Private/Materials/MeshMaterialShader.cpp
@@ -87,9 +87,11 @@ FShader* FMeshMaterialShaderType::FinishCompileShader(
 	check(CurrentJob.bSucceeded);
 	check(CurrentJob.VFType);
 
+	FShaderType* SpecificType = CurrentJob.ShaderType->LimitShaderResourceToThisType() ? CurrentJob.ShaderType : NULL;
+
 	// Reuse an existing resource with the same key or create a new one based on the compile output
 	// This allows FShaders to share compiled bytecode and RHI shader references
-	FShaderResource* Resource = FShaderResource::FindOrCreateShaderResource(CurrentJob.Output);
+	FShaderResource* Resource = FShaderResource::FindOrCreateShaderResource(CurrentJob.Output, SpecificType);
 
 	// Find a shader with the same key in memory
 	FShader* Shader = CurrentJob.ShaderType->FindShaderById(FShaderId(MaterialShaderMapHash, CurrentJob.VFType, CurrentJob.ShaderType, CurrentJob.Input.Target));
diff --git a/Engine/Source/Runtime/Engine/Private/ShaderDerivedDataVersion.h b/Engine/Source/Runtime/Engine/Private/ShaderDerivedDataVersion.h
--- a/Engine/Source/Runtime/Engine/Private/ShaderDerivedDataVersion.h
+++ b/Engine/Source/Runtime/Engine/Private/ShaderDerivedDataVersion.h
@@ -9,5 +9,5 @@
 // In case of merge conflicts with DDC versions, you MUST generate a new GUID and set this new
 // guid as version
 
-#define GLOBALSHADERMAP_DERIVEDDATA_VER			TEXT("af1606fcec114373844d2b6cf9df71c8")
-#define MATERIALSHADERMAP_DERIVEDDATA_VER		TEXT("bfbd0e1a93dd4481a991041cedfae4e9")
+#define GLOBALSHADERMAP_DERIVEDDATA_VER			TEXT("af1606fcec214373844d2b6cf9df71c8")
+#define MATERIALSHADERMAP_DERIVEDDATA_VER		TEXT("bfbd0e1a93dd4482a991041cedfae4e9")
diff --git a/Engine/Source/Runtime/Engine/Public/GlobalShader.h b/Engine/Source/Runtime/Engine/Public/GlobalShader.h
--- a/Engine/Source/Runtime/Engine/Public/GlobalShader.h
+++ b/Engine/Source/Runtime/Engine/Public/GlobalShader.h
@@ -52,9 +52,10 @@ public:
 		ConstructSerializedType InConstructSerializedRef,
 		ConstructCompiledType InConstructCompiledRef,
 		ModifyCompilationEnvironmentType InModifyCompilationEnvironmentRef,
-		ShouldCacheType InShouldCacheRef
+		ShouldCacheType InShouldCacheRef,
+		GetStreamOutElementsType InGetStreamOutElementsRef
 		):
-		FShaderType(InName,InSourceFilename,InFunctionName,InFrequency,InConstructSerializedRef),
+		FShaderType(InName,InSourceFilename,InFunctionName,InFrequency,InConstructSerializedRef,InGetStreamOutElementsRef),
 		ConstructCompiledRef(InConstructCompiledRef),
 		ShouldCacheRef(InShouldCacheRef),
 		ModifyCompilationEnvironmentRef(InModifyCompilationEnvironmentRef)
diff --git a/Engine/Source/Runtime/Engine/Public/MaterialShaderType.h b/Engine/Source/Runtime/Engine/Public/MaterialShaderType.h
--- a/Engine/Source/Runtime/Engine/Public/MaterialShaderType.h
+++ b/Engine/Source/Runtime/Engine/Public/MaterialShaderType.h
@@ -78,9 +78,10 @@ public:
 		ConstructSerializedType InConstructSerializedRef,
 		ConstructCompiledType InConstructCompiledRef,
 		ModifyCompilationEnvironmentType InModifyCompilationEnvironmentRef,
-		ShouldCacheType InShouldCacheRef
+		ShouldCacheType InShouldCacheRef,
+		GetStreamOutElementsType InGetStreamOutElementsRef
 		):
-		FShaderType(InName,InSourceFilename,InFunctionName,InFrequency,InConstructSerializedRef),
+		FShaderType(InName,InSourceFilename,InFunctionName,InFrequency,InConstructSerializedRef,InGetStreamOutElementsRef),
 		ConstructCompiledRef(InConstructCompiledRef),
 		ShouldCacheRef(InShouldCacheRef),
 		ModifyCompilationEnvironmentRef(InModifyCompilationEnvironmentRef)
diff --git a/Engine/Source/Runtime/Engine/Public/MeshMaterialShaderType.h b/Engine/Source/Runtime/Engine/Public/MeshMaterialShaderType.h
--- a/Engine/Source/Runtime/Engine/Public/MeshMaterialShaderType.h
+++ b/Engine/Source/Runtime/Engine/Public/MeshMaterialShaderType.h
@@ -49,9 +49,10 @@ public:
 		ConstructSerializedType InConstructSerializedRef,
 		ConstructCompiledType InConstructCompiledRef,
 		ModifyCompilationEnvironmentType InModifyCompilationEnvironmentRef,
-		ShouldCacheType InShouldCacheRef
+		ShouldCacheType InShouldCacheRef,
+		GetStreamOutElementsType InGetStreamOutElementsRef
 		):
-		FShaderType(InName,InSourceFilename,InFunctionName,InFrequency,InConstructSerializedRef),
+		FShaderType(InName,InSourceFilename,InFunctionName,InFrequency,InConstructSerializedRef,InGetStreamOutElementsRef),
 		ConstructCompiledRef(InConstructCompiledRef),
 		ShouldCacheRef(InShouldCacheRef),
 		ModifyCompilationEnvironmentRef(InModifyCompilationEnvironmentRef)
diff --git a/Engine/Source/Runtime/Renderer/Private/ConvertToUniformMesh.cpp b/Engine/Source/Runtime/Renderer/Private/ConvertToUniformMesh.cpp
--- a/Engine/Source/Runtime/Renderer/Private/ConvertToUniformMesh.cpp
+++ b/Engine/Source/Runtime/Renderer/Private/ConvertToUniformMesh.cpp
@@ -56,7 +56,16 @@ public:
 
 IMPLEMENT_MATERIAL_SHADER_TYPE(,FConvertToUniformMeshVS,TEXT("ConvertToUniformMesh"),TEXT("ConvertToUniformMeshVS"),SF_Vertex); 
 
-FStreamOutElementList GUniformMeshVertexLayout;
+void GetUniformMeshStreamOutLayout(FStreamOutElementList& Layout)
+{
+	Layout.Add(FStreamOutElement(0, "SV_Position", 0, 4, 0));
+	Layout.Add(FStreamOutElement(0, "Tangent", 0, 3, 0));
+	Layout.Add(FStreamOutElement(0, "Tangent", 1, 3, 0));
+	Layout.Add(FStreamOutElement(0, "Tangent", 2, 3, 0));
+	Layout.Add(FStreamOutElement(0, "UV", 0, 2, 0));
+	Layout.Add(FStreamOutElement(0, "UV", 1, 2, 0));
+	Layout.Add(FStreamOutElement(0, "VertexColor", 0, 4, 0));
+}
 
 // In float4's, must match usf
 int32 FSurfelBuffers::InterpolatedVertexDataStride = 6;
@@ -64,22 +73,14 @@ int32 FSurfelBuffers::InterpolatedVertexDataStride = 6;
 /** Returns number of float's in the uniform vertex. */
 int32 ComputeUniformVertexStride()
 {
-	if (GUniformMeshVertexLayout.Num() == 0)
-	{
-		GUniformMeshVertexLayout.Add(FStreamOutElement(0, "SV_Position", 0, 4, 0));
-		GUniformMeshVertexLayout.Add(FStreamOutElement(0, "Tangent", 0, 3, 0));
-		GUniformMeshVertexLayout.Add(FStreamOutElement(0, "Tangent", 1, 3, 0));
-		GUniformMeshVertexLayout.Add(FStreamOutElement(0, "Tangent", 2, 3, 0));
-		GUniformMeshVertexLayout.Add(FStreamOutElement(0, "UV", 0, 2, 0));
-		GUniformMeshVertexLayout.Add(FStreamOutElement(0, "UV", 1, 2, 0));
-		GUniformMeshVertexLayout.Add(FStreamOutElement(0, "VertexColor", 0, 4, 0));
-	}
-
+	FStreamOutElementList Layout;
 	int32 StreamStride = 0;
 
-	for (int32 ElementIndex = 0; ElementIndex < GUniformMeshVertexLayout.Num(); ElementIndex++)
+	GetUniformMeshStreamOutLayout(Layout);
+
+	for (int32 ElementIndex = 0; ElementIndex < Layout.Num(); ElementIndex++)
 	{
-		StreamStride += GUniformMeshVertexLayout[ElementIndex].ComponentCount;
+		StreamStride += Layout[ElementIndex].ComponentCount;
 	}
 
 	// D3D11 stream out buffer element stride must be a factor of 4
@@ -123,6 +124,13 @@ protected:
 				|| FCString::Strstr(VertexFactoryType->GetName(), TEXT("InstancedStaticMeshVertexFactory")) != NULL);
 	}
 
+	static void GetStreamOutElements(FStreamOutElementList& ElementList, TArray<uint32>& StreamStrides, int32& RasterizedStream)
+	{
+		StreamStrides.Add(ComputeUniformVertexStride() * 4);
+		GetUniformMeshStreamOutLayout(ElementList);
+		RasterizedStream = -1;
+	}
+
 public:
 	
 	void SetParameters(FRHICommandList& RHICmdList, const FVertexFactory* VertexFactory,const FMaterialRenderProxy* MaterialRenderProxy,const FSceneView* View)
@@ -134,13 +142,6 @@ public:
 	{
 		FMeshMaterialShader::SetMesh(RHICmdList, (FGeometryShaderRHIParamRef)GetGeometryShader(),VertexFactory,View,Proxy,BatchElement);
 	}
-
-	virtual void GetStreamOutElements(FStreamOutElementList& ElementList, TArray<uint32>& StreamStrides, int32& RasterizedStream) override
-	{
-		StreamStrides.Add(ComputeUniformVertexStride() * 4);
-		ElementList = GUniformMeshVertexLayout;
-		RasterizedStream = -1;
-	}
 };
 
 IMPLEMENT_MATERIAL_SHADER_TYPE(,FConvertToUniformMeshGS,TEXT("ConvertToUniformMesh"),TEXT("ConvertToUniformMeshGS"),SF_Geometry); 
diff --git a/Engine/Source/Runtime/ShaderCore/Private/Shader.cpp b/Engine/Source/Runtime/ShaderCore/Private/Shader.cpp
--- a/Engine/Source/Runtime/ShaderCore/Private/Shader.cpp
+++ b/Engine/Source/Runtime/ShaderCore/Private/Shader.cpp
@@ -75,13 +75,15 @@ FShaderType::FShaderType(
 	const TCHAR* InSourceFilename,
 	const TCHAR* InFunctionName,
 	uint32 InFrequency,
-	ConstructSerializedType InConstructSerializedRef
+	ConstructSerializedType InConstructSerializedRef,
+	GetStreamOutElementsType InGetStreamOutElementsRef
 	):
 	Name(InName),
 	SourceFilename(InSourceFilename),
 	FunctionName(InFunctionName),
 	Frequency(InFrequency),
 	ConstructSerializedRef(InConstructSerializedRef),
+	GetStreamOutElementsRef(InGetStreamOutElementsRef),
 	GlobalListLink(this)
 {
 	for (int32 Platform = 0; Platform < SP_NumPlatforms; Platform++)
@@ -285,7 +287,8 @@ void FShaderType::Uninitialize()
 TMap<FShaderResourceId, FShaderResource*> FShaderResource::ShaderResourceIdMap;
 
 FShaderResource::FShaderResource()
-	: NumInstructions(0)
+	: SpecificType(NULL)
+	, NumInstructions(0)
 	, NumTextureSamplers(0)
 	, NumRefs(0)
 {
@@ -293,8 +296,9 @@ FShaderResource::FShaderResource()
 }
 
 
-FShaderResource::FShaderResource(const FShaderCompilerOutput& Output) 
-	: NumInstructions(Output.NumInstructions)
+FShaderResource::FShaderResource(const FShaderCompilerOutput& Output, FShaderType* InSpecificType) 
+	: SpecificType(InSpecificType)
+	, NumInstructions(Output.NumInstructions)
 	, NumTextureSamplers(Output.NumTextureSamplers)
 	, NumRefs(0)
 	
@@ -329,6 +333,7 @@ void FShaderResource::Register()
 
 void FShaderResource::Serialize(FArchive& Ar)
 {
+	Ar << SpecificType;
 	Ar << Target;
 	Ar << Code;
 	Ar << OutputHash;
@@ -376,14 +381,14 @@ FShaderResource* FShaderResource::FindShaderResourceById(const FShaderResourceId
 }
 
 
-FShaderResource* FShaderResource::FindOrCreateShaderResource(const FShaderCompilerOutput& Output)
+FShaderResource* FShaderResource::FindOrCreateShaderResource(const FShaderCompilerOutput& Output, FShaderType* SpecificType)
 {
-	const FShaderResourceId ResourceId(Output);
+	const FShaderResourceId ResourceId(Output, SpecificType ? SpecificType->GetName() : NULL);
 	FShaderResource* Resource = FindShaderResourceById(ResourceId);
 
 	if (!Resource)
 	{
-		Resource = new FShaderResource(Output);
+		Resource = new FShaderResource(Output, SpecificType);
 	}
 
 	return Resource;
@@ -470,13 +475,32 @@ void FShaderResource::InitRHI()
 	}
 	else if(Target.Frequency == SF_Geometry)
 	{
-		GeometryShader = ShaderCache ? ShaderCache->GetGeometryShader((EShaderPlatform)Target.Platform, OutputHash, Code) : RHICreateGeometryShader(Code);
+		if (SpecificType)
+		{
+			FStreamOutElementList ElementList;
+			TArray<uint32> StreamStrides;
+			int32 RasterizedStream = -1;
+			SpecificType->GetStreamOutElements(ElementList, StreamStrides, RasterizedStream);
+			checkf(ElementList.Num(), *FString::Printf(TEXT("Shader type %s was given GetStreamOutElements implementation that had no elements!"), SpecificType->GetName()));
+
+			//@todo - not using the cache
+			GeometryShader = RHICreateGeometryShaderWithStreamOutput(Code, ElementList, StreamStrides.Num(), StreamStrides.GetData(), RasterizedStream);
+		}
+		else
+		{
+			GeometryShader = ShaderCache ? ShaderCache->GetGeometryShader((EShaderPlatform)Target.Platform, OutputHash, Code) : RHICreateGeometryShader(Code);
+		}
 	}
 	else if(Target.Frequency == SF_Compute)
 	{
 		ComputeShader = ShaderCache ? ShaderCache->GetComputeShader((EShaderPlatform)Target.Platform, OutputHash, Code) : RHICreateComputeShader(Code);
 	}
 
+	if (Target.Frequency != SF_Geometry)
+	{
+		checkf(!SpecificType, *FString::Printf(TEXT("Only geometry shaders can use GetStreamOutElements, shader type %s"), SpecificType->GetName()));
+	}
+
 	if (!FPlatformProperties::HasEditorOnlyData())
 	{
 		DEC_DWORD_STAT_BY_FName(GetMemoryStatType((EShaderFrequency)Target.Frequency).GetName(), Code.Num());
@@ -627,6 +651,14 @@ const FComputeShaderRHIRef& FShaderResource::GetComputeShader()
 	return ComputeShader; 
 }
 
+FShaderResourceId FShaderResource::GetId() const
+{
+	FShaderResourceId ShaderId;
+	ShaderId.Target = Target;
+	ShaderId.OutputHash = OutputHash;
+	ShaderId.SpecificShaderTypeName = SpecificType ? SpecificType->GetName() : NULL;
+	return ShaderId;
+}
 
 FShaderId::FShaderId(const FSHAHash& InMaterialShaderMapHash, FVertexFactoryType* InVertexFactoryType, FShaderType* InShaderType, FShaderTarget InTarget)
 	: MaterialShaderMapHash(InMaterialShaderMapHash)
@@ -761,7 +793,6 @@ FShader::FShader(const CompiledShaderInitializerType& Initializer):
 
 FShader::~FShader()
 {
-	check(!IsInitialized());
 	check(Canary == ShaderMagic_Uninitialized || Canary == ShaderMagic_Initialized);
 	check(NumRefs == 0);
 	Canary = 0;
@@ -898,8 +929,6 @@ void FShader::Release()
 		// Deregister the shader now to eliminate references to it by the type's ShaderIdMap
 		Deregister();
 
-		BeginReleaseResource(this);
-
 		BeginCleanup(this);
 	}
 }
@@ -947,42 +976,6 @@ void FShader::FinishCleanup()
 	delete this;
 }
 
-
-void FShader::InitRHI()
-{
-	checkf(Resource->GetCode().Num() > 0, TEXT("FShader::InitRHI was called with empty bytecode, which can happen if the resource is initialized multiple times on platforms with no editor data."));
-
-	// we can't have this called on the wrong platform's shaders
-	if (!ArePlatformsCompatible(GMaxRHIShaderPlatform, (EShaderPlatform)Target.Platform))
- 	{
- 		if (FPlatformProperties::RequiresCookedData())
- 		{
- 			UE_LOG(LogShaders, Fatal, TEXT("FShader::InitRHI got platform %s but it is not compatible with %s"), 
-				*LegacyShaderPlatformToShaderFormat((EShaderPlatform)Target.Platform).ToString(), *LegacyShaderPlatformToShaderFormat(GMaxRHIShaderPlatform).ToString());
- 		}
- 		return;
- 	}
-
-	if (Target.Frequency == SF_Geometry)
-	{
-		FStreamOutElementList ElementList;
-		TArray<uint32> StreamStrides;
-		int32 RasterizedStream = -1;
-		GetStreamOutElements(ElementList, StreamStrides, RasterizedStream);
-
-		if (ElementList.Num() > 0)
-		{
-			GeometryShaderWithStreamOutput = RHICreateGeometryShaderWithStreamOutput(Resource->GetCode(), ElementList, StreamStrides.Num(), StreamStrides.GetData(), RasterizedStream);
-		}
-	}
-}
-
-
-void FShader::ReleaseRHI()
-{
-	GeometryShaderWithStreamOutput.SafeRelease();
-}
-
 void FShader::VerifyBoundUniformBufferParameters()
 {
 	// Support being called on a NULL pointer
diff --git a/Engine/Source/Runtime/ShaderCore/Public/Shader.h b/Engine/Source/Runtime/ShaderCore/Public/Shader.h
--- a/Engine/Source/Runtime/ShaderCore/Public/Shader.h
+++ b/Engine/Source/Runtime/ShaderCore/Public/Shader.h
@@ -29,9 +29,10 @@ public:
 
 	FShaderResourceId() {}
 
-	FShaderResourceId(const FShaderCompilerOutput& Output) :
+	FShaderResourceId(const FShaderCompilerOutput& Output, const TCHAR* InSpecificShaderTypeName) :
 		Target(Output.Target),
-		OutputHash(Output.OutputHash)
+		OutputHash(Output.OutputHash),
+		SpecificShaderTypeName(InSpecificShaderTypeName)
 	{}
 
 	friend inline uint32 GetTypeHash( const FShaderResourceId& Id )
@@ -41,7 +42,10 @@ public:
 
 	friend bool operator==(const FShaderResourceId& X, const FShaderResourceId& Y)
 	{
-		return X.Target == Y.Target && X.OutputHash == Y.OutputHash;
+		return X.Target == Y.Target 
+			&& X.OutputHash == Y.OutputHash 
+			&& ((X.SpecificShaderTypeName == NULL && Y.SpecificShaderTypeName == NULL)
+				|| (FCString::Strcmp(X.SpecificShaderTypeName, Y.SpecificShaderTypeName) == 0));
 	}
 
 	friend bool operator!=(const FShaderResourceId& X, const FShaderResourceId& Y)
@@ -52,6 +56,25 @@ public:
 	friend FArchive& operator<<(FArchive& Ar, FShaderResourceId& Id)
 	{
 		Ar << Id.Target << Id.OutputHash;
+
+		if (Ar.IsSaving())
+		{
+			Id.SpecificShaderTypeStorage = Id.SpecificShaderTypeName ? Id.SpecificShaderTypeName : TEXT("");
+		}
+
+		Ar << Id.SpecificShaderTypeStorage;
+
+		if (Ar.IsLoading())
+		{
+			Id.SpecificShaderTypeName = *Id.SpecificShaderTypeStorage;
+
+			if (FCString::Strcmp(Id.SpecificShaderTypeName, TEXT("")) == 0)
+			{
+				// Store NULL for empty string to be consistent with FShaderResourceId's created at compile time
+				Id.SpecificShaderTypeName = NULL;
+			}
+		}
+
 		return Ar;
 	}
 
@@ -60,6 +83,12 @@ public:
 
 	/** Hash of the compiled shader output, which is used to create the FShaderResource. */
 	FSHAHash OutputHash;
+
+	/** NULL if type doesn't matter, otherwise the name of the type that this was created specifically for, which is used with geometry shader stream out. */
+	const TCHAR* SpecificShaderTypeName;
+
+	/** Stores the memory for SpecificShaderTypeName if this is a standalone Id, otherwise is empty and SpecificShaderTypeName points to an FShaderType name. */
+	FString SpecificShaderTypeStorage;
 };
 
 /** 
@@ -75,7 +104,7 @@ public:
 	SHADERCORE_API FShaderResource();
 
 	/** Constructor used when creating a new shader resource from compiled output. */
-	FShaderResource(const FShaderCompilerOutput& Output);
+	FShaderResource(const FShaderCompilerOutput& Output, FShaderType* InSpecificType);
 
 	~FShaderResource();
 
@@ -116,24 +145,13 @@ public:
 	/** @return the shader's compute shader */
 	const FComputeShaderRHIRef& GetComputeShader();
 
-	FShaderResourceId GetId() const
-	{
-		FShaderResourceId ShaderId;
-		ShaderId.Target = Target;
-		ShaderId.OutputHash = OutputHash;
-		return ShaderId;
-	}
+	SHADERCORE_API FShaderResourceId GetId() const;
 
 	uint32 GetSizeBytes() const
 	{
 		return Code.GetAllocatedSize() + sizeof(FShaderResource);
 	}
 
-	const TArray<uint8>& GetCode() const
-	{
-		return Code;
-	}
-
 	// FRenderResource interface.
 	virtual void InitRHI();
 	virtual void ReleaseRHI();
@@ -144,8 +162,11 @@ public:
 	/** Finds a matching shader resource in memory if possible. */
 	SHADERCORE_API static FShaderResource* FindShaderResourceById(const FShaderResourceId& Id);
 
-	/** Finds a matching shader resource in memory or creates a new one with the given compiler output. */
-	SHADERCORE_API static FShaderResource* FindOrCreateShaderResource(const FShaderCompilerOutput& Output);
+	/** 
+	 * Finds a matching shader resource in memory or creates a new one with the given compiler output.  
+	 * SpecificType can be NULL
+	 */
+	SHADERCORE_API static FShaderResource* FindOrCreateShaderResource(const FShaderCompilerOutput& Output, class FShaderType* SpecificType);
 
 	/** Return a list of all shader Ids currently known */
 	SHADERCORE_API static void GetAllShaderResourceId(TArray<FShaderResourceId>& Ids);
@@ -172,6 +193,9 @@ private:
 	 */
 	FSHAHash OutputHash;
 
+	/** If not NULL, the shader type this resource must be used with. */
+	class FShaderType* SpecificType;
+
 	/** The number of instructions the shader takes to execute. */
 	uint32 NumInstructions;
 
@@ -386,7 +410,7 @@ public:
 class FMaterial;
 
 /** A compiled shader and its parameter bindings. */
-class SHADERCORE_API FShader : public FRenderResource, public FDeferredCleanupInterface
+class SHADERCORE_API FShader : public FDeferredCleanupInterface
 {
 	friend class FShaderType;
 public:
@@ -482,12 +506,7 @@ public:
 	/** @return the shader's geometry shader */
 	const FGeometryShaderRHIRef& GetGeometryShader()
 	{
-		if (!IsInitialized())
-		{
-			InitResourceFromPossiblyParallelRendering();
-		}
-
-		return GeometryShaderWithStreamOutput ? GeometryShaderWithStreamOutput : Resource->GetGeometryShader();
+		return Resource->GetGeometryShader();
 	}
 	/** @return the shader's compute shader */
 	const FComputeShaderRHIRef& GetComputeShader()
@@ -539,14 +558,10 @@ public:
 	// FDeferredCleanupInterface implementation.
 	virtual void FinishCleanup();
 
-	// FRenderResource interface.
-	virtual void InitRHI();
-	virtual void ReleaseRHI();
-
 	/** Implement for geometry shaders that want to use stream out. */
-	virtual void GetStreamOutElements(FStreamOutElementList& ElementList, TArray<uint32>& StreamStrides, int32& RasterizedStream) {}
+	static void GetStreamOutElements(FStreamOutElementList& ElementList, TArray<uint32>& StreamStrides, int32& RasterizedStream) {}
 
-	void InitializeResource()
+	void BeginInitializeResources()
 	{
 		BeginInitResource(Resource);
 	}
@@ -641,8 +656,6 @@ private:
 	/** Reference to the shader resource, which stores the compiled bytecode and the RHI shader resource. */
 	TRefCountPtr<FShaderResource> Resource;
 
-	FGeometryShaderRHIRef GeometryShaderWithStreamOutput;
-
 	/** Hash of the material shader map this shader belongs to, stored so that an FShaderId can be constructed from this shader. */
 	FSHAHash MaterialShaderMapHash;
 
@@ -684,6 +697,7 @@ class SHADERCORE_API FShaderType
 {
 public:
 	typedef class FShader* (*ConstructSerializedType)();
+	typedef void (*GetStreamOutElementsType)(FStreamOutElementList& ElementList, TArray<uint32>& StreamStrides, int32& RasterizedStream);
 
 	/** @return The global shader factory list. */
 	static TLinkedList<FShaderType*>*& GetTypeList();
@@ -709,7 +723,8 @@ public:
 		const TCHAR* InSourceFilename,
 		const TCHAR* InFunctionName,
 		uint32 InFrequency,
-		ConstructSerializedType InConstructSerializedRef);
+		ConstructSerializedType InConstructSerializedRef,
+		GetStreamOutElementsType InGetStreamOutElementsRef);
 
 	virtual ~FShaderType();
 
@@ -792,6 +807,16 @@ public:
 		return ShaderIdMap;
 	}
 
+	bool LimitShaderResourceToThisType()
+	{
+		return GetStreamOutElementsRef != &FShader::GetStreamOutElements;
+	}
+
+	void GetStreamOutElements(FStreamOutElementList& ElementList, TArray<uint32>& StreamStrides, int32& RasterizedStream) 
+	{
+		(*GetStreamOutElementsRef)(ElementList, StreamStrides, RasterizedStream);
+	}
+
 private:
 
 	uint32 HashIndex;
@@ -801,6 +826,7 @@ private:
 	uint32 Frequency;
 
 	ConstructSerializedType ConstructSerializedRef;
+	GetStreamOutElementsType GetStreamOutElementsRef;
 
 	/** A map from shader ID to shader.  A shader will be removed from it when deleted, so this doesn't need to use a TRefCountPtr. */
 	TMap<FShaderId,FShader*> ShaderIdMap;
@@ -859,7 +885,8 @@ private:
 		ShaderClass::ConstructSerializedInstance, \
 		ShaderClass::ConstructCompiledInstance, \
 		ShaderClass::ModifyCompilationEnvironment, \
-		ShaderClass::ShouldCache \
+		ShaderClass::ShouldCache, \
+		ShaderClass::GetStreamOutElements \
 		);
 
 /** A macro to implement a templated shader type, the function name and the source filename comes from the class. */
@@ -873,7 +900,8 @@ private:
 	ShaderClass::ConstructSerializedInstance, \
 	ShaderClass::ConstructCompiledInstance, \
 	ShaderClass::ModifyCompilationEnvironment, \
-	ShaderClass::ShouldCache \
+	ShaderClass::ShouldCache, \
+	ShaderClass::GetStreamOutElements \
 	);
 
 
@@ -887,7 +915,8 @@ private:
 	ShaderClass::ConstructSerializedInstance, \
 	ShaderClass::ConstructCompiledInstance, \
 	ShaderClass::ModifyCompilationEnvironment, \
-	ShaderClass::ShouldCache \
+	ShaderClass::ShouldCache, \
+	ShaderClass::GetStreamOutElements \
 	);
 #endif
 
