Fixed FCompareShaderTypes so that shader types of different lengths don't accidently compare to the same value...

diff --git a/Engine/Source/Runtime/ShaderCore/Public/Shader.h b/Engine/Source/Runtime/ShaderCore/Public/Shader.h
--- a/Engine/Source/Runtime/ShaderCore/Public/Shader.h
+++ b/Engine/Source/Runtime/ShaderCore/Public/Shader.h
@@ -953,7 +953,13 @@ class FCompareShaderTypes
 public:		
 	FORCEINLINE bool operator()(const FShaderType& A, const FShaderType& B ) const
 	{
-		return FCString::Strncmp(A.GetName(), B.GetName(), FMath::Min(FCString::Strlen(A.GetName()), FCString::Strlen(B.GetName()))) > 0;
+		int32 AL = FCString::Strlen(A.GetName());
+		int32 BL = FCString::Strlen(B.GetName());
+		if ( AL == BL )
+		{
+			return FCString::Strncmp(A.GetName(), B.GetName(), AL) > 0;
+		}
+		return AL > BL;
 	}
 };
 
