Shader caching has been updated to include the option to record draw-states used for each shader so that they may be pre-drawn to reduce in-driver hitching.
- This is especially important on OS X OpenGL where these hitches can be equivalent to several frames, so it is enabled by default for this platform.
- The logging of draw calls is turned on/off with r.UseShaderDrawLog and incurs a fixed overhead as all RHI state is tracked & recorded into the shader cache.
- The actual predrawing is turned on/off with r.UseShaderPredraw and requires an existing cache to be useful, the predrawing takes place at frame-end and can be controlled with r.PredrawBatchTime which the time in ms to spend predrawing each frame until complete, or -1 to do it in one go (currently the default).
- The shader cache is now also versioned with the game able to supply a version distinct from the engine, so that old caches may be invalidated if/when required.
- The shader cache is now compatible with OpenGL SM4 & SM5 on Windows & probably Linux too.
- It may not work correctly with desktop ES2 emulation as that is still untested.
- To make it easier to add to other RHIs the API has been made simper to avoid lots of exposed branches.

diff --git a/Engine/Source/Developer/ShaderFormatOpenGL/Private/OpenGLShaderCompiler.cpp b/Engine/Source/Developer/ShaderFormatOpenGL/Private/OpenGLShaderCompiler.cpp
--- a/Engine/Source/Developer/ShaderFormatOpenGL/Private/OpenGLShaderCompiler.cpp
+++ b/Engine/Source/Developer/ShaderFormatOpenGL/Private/OpenGLShaderCompiler.cpp
@@ -494,6 +494,23 @@ void ParseGlslError(TArray<FShaderCompilerError>& OutErrors, const FString& InLi
 	}
 }
 
+static TArray<ANSICHAR> ParseIdentifierANSI(const ANSICHAR* &Str)
+{
+	TArray<ANSICHAR> Result;
+	
+	while ((*Str >= 'A' && *Str <= 'Z')
+		   || (*Str >= 'a' && *Str <= 'z')
+		   || (*Str >= '0' && *Str <= '9')
+		   || *Str == '_')
+	{
+		Result.Add(FChar::ToLower(*Str));
+		++Str;
+	}
+	Result.Add('\0');
+	
+	return Result;
+}
+
 static FString ParseIdentifier(const ANSICHAR* &Str)
 {
 	FString Result;
@@ -630,7 +647,7 @@ static void BuildShaderOutput(
             {
                 FOpenGLShaderVarying Var;
                 Var.Location = Location;
-                Var.Varying = ParseIdentifier(ShaderSource);
+                Var.Varying = ParseIdentifierANSI(ShaderSource);
                 Header.Bindings.InputVaryings.Add(Var);
             }
 
@@ -683,7 +700,7 @@ static void BuildShaderOutput(
             {
                 FOpenGLShaderVarying Var;
                 Var.Location = Location;
-                Var.Varying = ParseIdentifier(ShaderSource);
+                Var.Varying = ParseIdentifierANSI(ShaderSource);
                 Header.Bindings.OutputVaryings.Add(Var);
             }
 
diff --git a/Engine/Source/Runtime/Launch/Private/LaunchEngineLoop.cpp b/Engine/Source/Runtime/Launch/Private/LaunchEngineLoop.cpp
--- a/Engine/Source/Runtime/Launch/Private/LaunchEngineLoop.cpp
+++ b/Engine/Source/Runtime/Launch/Private/LaunchEngineLoop.cpp
@@ -34,6 +34,7 @@
 	#include "DerivedDataCacheInterface.h"
 	#include "RenderCore.h"
 	#include "ShaderCompiler.h"
+	#include "ShaderCache.h"
 	#include "DistanceFieldAtlas.h"
 	#include "GlobalShader.h"
 	#include "ParticleHelper.h"
@@ -2368,6 +2369,7 @@ void FEngineLoop::Tick()
 		ENQUEUE_UNIQUE_RENDER_COMMAND(
 			EndFrame,
 		{
+			FShaderCache::PreDrawShaders(RHICmdList);
 			RHICmdList.EndFrame();
 			RHICmdList.PopEvent();
 		});
diff --git a/Engine/Source/Runtime/OpenGLDrv/Private/Linux/OpenGLLinux.h b/Engine/Source/Runtime/OpenGLDrv/Private/Linux/OpenGLLinux.h
--- a/Engine/Source/Runtime/OpenGLDrv/Private/Linux/OpenGLLinux.h
+++ b/Engine/Source/Runtime/OpenGLDrv/Private/Linux/OpenGLLinux.h
@@ -348,7 +348,8 @@
 	EnumMacro(PFNGLPROGRAMUNIFORM4UIVPROC, glProgramUniform4uiv)\
 	EnumMacro(PFNGLGETPROGRAMPIPELINEIVPROC, glGetProgramPipelineiv)\
 	EnumMacro(PFNGLVALIDATEPROGRAMPIPELINEPROC, glValidateProgramPipeline)\
-	EnumMacro(PFNGLGETPROGRAMPIPELINEINFOLOGPROC, glGetProgramPipelineInfoLog)
+	EnumMacro(PFNGLGETPROGRAMPIPELINEINFOLOGPROC, glGetProgramPipelineInfoLog)\
+	EnumMacro(PFNGLISPROGRAMPIPELINEPROC, glIsProgramPipeline)
 
 /** List of all OpenGL entry points. */
 #define ENUM_GL_ENTRYPOINTS_ALL(EnumMacro) \
diff --git a/Engine/Source/Runtime/OpenGLDrv/Private/OpenGLCommands.cpp b/Engine/Source/Runtime/OpenGLDrv/Private/OpenGLCommands.cpp
--- a/Engine/Source/Runtime/OpenGLDrv/Private/OpenGLCommands.cpp
+++ b/Engine/Source/Runtime/OpenGLDrv/Private/OpenGLCommands.cpp
@@ -6,6 +6,7 @@
 
 #include "OpenGLDrvPrivate.h"
 #include "OpenGLState.h"
+#include "ShaderCache.h"
 
 #define DECLARE_ISBOUNDSHADER(ShaderType) inline void ValidateBoundShader(TRefCountPtr<FOpenGLBoundShaderState> InBoundShaderState, F##ShaderType##RHIParamRef ShaderType##RHI) \
 { \
@@ -291,6 +292,8 @@ void FOpenGLDynamicRHI::RHISetRasterizerState(FRasterizerStateRHIParamRef NewSta
 	VERIFY_GL_SCOPE();
 	DYNAMIC_CAST_OPENGLRESOURCE(RasterizerState,NewState);
 	PendingState.RasterizerState = NewState->Data;
+	
+	FShaderCache::SetRasterizerState(NewStateRHI);
 }
 
 void FOpenGLDynamicRHI::UpdateRasterizerStateInOpenGLContext( FOpenGLContextState& ContextState )
@@ -385,6 +388,8 @@ void FOpenGLDynamicRHI::RHISetViewport(uint32 MinX,uint32 MinY,float MinZ,uint32
 	PendingState.Viewport.Max.Y = MaxY;
 	PendingState.DepthMinZ = MinZ;
 	PendingState.DepthMaxZ = MaxZ;
+	
+	FShaderCache::SetViewport(MinX, MinY, MinZ, MaxX, MaxY, MaxZ);
 }
 
 void FOpenGLDynamicRHI::RHISetScissorRect(bool bEnable,uint32 MinX,uint32 MinY,uint32 MaxX,uint32 MaxY)
@@ -435,6 +440,8 @@ void FOpenGLDynamicRHI::RHISetBoundShaderState( FBoundShaderStateRHIParamRef Bou
 	// Prevent transient bound shader states from being recreated for each use by keeping a history of the most recently used bound shader states.
 	// The history keeps them alive, and the bound shader state cache allows them to be reused if needed.
 	BoundShaderStateHistory.Add(BoundShaderState);
+	
+	FShaderCache::SetBoundShaderState(BoundShaderStateRHI);
 }
 
 void FOpenGLDynamicRHI::RHISetUAVParameter(FComputeShaderRHIParamRef ComputeShaderRHI,uint32 UAVIndex,FUnorderedAccessViewRHIParamRef UnorderedAccessViewRHI)
@@ -801,6 +808,8 @@ void FOpenGLDynamicRHI::RHISetShaderResourceViewParameter(FPixelShaderRHIParamRe
 	}
 	InternalSetShaderTexture(NULL, SRV, FOpenGL::GetFirstPixelTextureUnit() + TextureIndex, Target, Resource, 0, LimitMip);
 	RHISetShaderSampler(PixelShaderRHI,TextureIndex,PointSamplerState);
+	
+	FShaderCache::SetSRV(SF_Pixel, TextureIndex, SRVRHI);
 }
 
 void FOpenGLDynamicRHI::RHISetShaderResourceViewParameter(FVertexShaderRHIParamRef VertexShaderRHI,uint32 TextureIndex,FShaderResourceViewRHIParamRef SRVRHI)
@@ -821,6 +830,8 @@ void FOpenGLDynamicRHI::RHISetShaderResourceViewParameter(FVertexShaderRHIParamR
 	}
 	InternalSetShaderTexture(NULL, SRV, FOpenGL::GetFirstVertexTextureUnit() + TextureIndex, Target, Resource, 0, LimitMip);
 	RHISetShaderSampler(VertexShaderRHI,TextureIndex,PointSamplerState);
+	
+	FShaderCache::SetSRV(SF_Vertex, TextureIndex, SRVRHI);
 }
 
 void FOpenGLDynamicRHI::RHISetShaderResourceViewParameter(FComputeShaderRHIParamRef ComputeShaderRHI,uint32 TextureIndex,FShaderResourceViewRHIParamRef SRVRHI)
@@ -840,6 +851,8 @@ void FOpenGLDynamicRHI::RHISetShaderResourceViewParameter(FComputeShaderRHIParam
 	}
 	InternalSetShaderTexture(NULL, SRV, FOpenGL::GetFirstComputeTextureUnit() + TextureIndex, Target, Resource, 0, LimitMip);
 	RHISetShaderSampler(ComputeShaderRHI,TextureIndex,PointSamplerState);
+	
+	FShaderCache::SetSRV(SF_Compute, TextureIndex, SRVRHI);
 }
 
 void FOpenGLDynamicRHI::RHISetShaderResourceViewParameter(FHullShaderRHIParamRef HullShaderRHI,uint32 TextureIndex,FShaderResourceViewRHIParamRef SRVRHI)
@@ -860,6 +873,8 @@ void FOpenGLDynamicRHI::RHISetShaderResourceViewParameter(FHullShaderRHIParamRef
 		LimitMip = SRV->LimitMip;
 	}
 	InternalSetShaderTexture(NULL, SRV, FOpenGL::GetFirstHullTextureUnit() + TextureIndex, Target, Resource, 0, LimitMip);
+	
+	FShaderCache::SetSRV(SF_Hull, TextureIndex, SRVRHI);
 }
 
 void FOpenGLDynamicRHI::RHISetShaderResourceViewParameter(FDomainShaderRHIParamRef DomainShaderRHI,uint32 TextureIndex,FShaderResourceViewRHIParamRef SRVRHI)
@@ -880,6 +895,8 @@ void FOpenGLDynamicRHI::RHISetShaderResourceViewParameter(FDomainShaderRHIParamR
 		LimitMip = SRV->LimitMip;
 	}
 	InternalSetShaderTexture(NULL, SRV, FOpenGL::GetFirstDomainTextureUnit() + TextureIndex, Target, Resource, 0, LimitMip);
+	
+	FShaderCache::SetSRV(SF_Domain, TextureIndex, SRVRHI);
 }
 
 void FOpenGLDynamicRHI::RHISetShaderResourceViewParameter(FGeometryShaderRHIParamRef GeometryShaderRHI,uint32 TextureIndex,FShaderResourceViewRHIParamRef SRVRHI)
@@ -900,6 +917,8 @@ void FOpenGLDynamicRHI::RHISetShaderResourceViewParameter(FGeometryShaderRHIPara
 	}
 	InternalSetShaderTexture(NULL, SRV, FOpenGL::GetFirstGeometryTextureUnit() + TextureIndex, Target, Resource, 0, LimitMip);
 	RHISetShaderSampler(GeometryShaderRHI,TextureIndex,PointSamplerState);
+	
+	FShaderCache::SetSRV(SF_Geometry, TextureIndex, SRVRHI);
 }
 
 void FOpenGLDynamicRHI::RHISetShaderTexture(FVertexShaderRHIParamRef VertexShaderRHI,uint32 TextureIndex,FTextureRHIParamRef NewTextureRHI)
@@ -916,6 +935,8 @@ void FOpenGLDynamicRHI::RHISetShaderTexture(FVertexShaderRHIParamRef VertexShade
 	{
 		InternalSetShaderTexture(nullptr, nullptr, FOpenGL::GetFirstVertexTextureUnit() + TextureIndex, 0, 0, 0, -1);
 	}
+	
+	FShaderCache::SetTexture(SF_Vertex, TextureIndex, NewTextureRHI);
 }
 
 void FOpenGLDynamicRHI::RHISetShaderTexture(FHullShaderRHIParamRef HullShaderRHI,uint32 TextureIndex,FTextureRHIParamRef NewTextureRHI)
@@ -933,6 +954,8 @@ void FOpenGLDynamicRHI::RHISetShaderTexture(FHullShaderRHIParamRef HullShaderRHI
 	{
 		InternalSetShaderTexture(nullptr, nullptr, FOpenGL::GetFirstHullTextureUnit() + TextureIndex, 0, 0, 0, -1);
 	}
+	
+	FShaderCache::SetTexture(SF_Hull, TextureIndex, NewTextureRHI);
 }
 
 void FOpenGLDynamicRHI::RHISetShaderTexture(FDomainShaderRHIParamRef DomainShaderRHI,uint32 TextureIndex,FTextureRHIParamRef NewTextureRHI)
@@ -950,6 +973,8 @@ void FOpenGLDynamicRHI::RHISetShaderTexture(FDomainShaderRHIParamRef DomainShade
 	{
 		InternalSetShaderTexture(nullptr, nullptr, FOpenGL::GetFirstDomainTextureUnit() + TextureIndex, 0, 0, 0, -1);
 	}
+	
+	FShaderCache::SetTexture(SF_Domain, TextureIndex, NewTextureRHI);
 }
 
 void FOpenGLDynamicRHI::RHISetShaderTexture(FGeometryShaderRHIParamRef GeometryShaderRHI,uint32 TextureIndex,FTextureRHIParamRef NewTextureRHI)
@@ -966,6 +991,8 @@ void FOpenGLDynamicRHI::RHISetShaderTexture(FGeometryShaderRHIParamRef GeometryS
 	{
 		InternalSetShaderTexture(nullptr, nullptr, FOpenGL::GetFirstGeometryTextureUnit() + TextureIndex, 0, 0, 0, -1);
 	}
+	
+	FShaderCache::SetTexture(SF_Geometry, TextureIndex, NewTextureRHI);
 }
 
 void FOpenGLDynamicRHI::RHISetShaderTexture(FPixelShaderRHIParamRef PixelShaderRHI,uint32 TextureIndex,FTextureRHIParamRef NewTextureRHI)
@@ -982,6 +1009,8 @@ void FOpenGLDynamicRHI::RHISetShaderTexture(FPixelShaderRHIParamRef PixelShaderR
 	{
 		InternalSetShaderTexture(nullptr, nullptr, FOpenGL::GetFirstPixelTextureUnit() + TextureIndex, 0, 0, 0, -1);
 	}
+	
+	FShaderCache::SetTexture(SF_Pixel, TextureIndex, NewTextureRHI);
 }
 
 void FOpenGLDynamicRHI::RHISetShaderSampler(FVertexShaderRHIParamRef VertexShaderRHI,uint32 SamplerIndex,FSamplerStateRHIParamRef NewStateRHI)
@@ -1005,6 +1034,8 @@ void FOpenGLDynamicRHI::RHISetShaderSampler(FVertexShaderRHIParamRef VertexShade
 	{
 		InternalSetSamplerStates(FOpenGL::GetFirstVertexTextureUnit() + SamplerIndex, NewState);
 	}
+	
+	FShaderCache::SetSamplerState(SF_Vertex, SamplerIndex, NewStateRHI);
 }
 
 void FOpenGLDynamicRHI::RHISetShaderSampler(FHullShaderRHIParamRef HullShaderRHI,uint32 SamplerIndex,FSamplerStateRHIParamRef NewStateRHI)
@@ -1029,6 +1060,8 @@ void FOpenGLDynamicRHI::RHISetShaderSampler(FHullShaderRHIParamRef HullShaderRHI
 	{
 		InternalSetSamplerStates(FOpenGL::GetFirstHullTextureUnit() + SamplerIndex, NewState);
 	}
+	
+	FShaderCache::SetSamplerState(SF_Hull, SamplerIndex, NewStateRHI);
 }
 void FOpenGLDynamicRHI::RHISetShaderSampler(FDomainShaderRHIParamRef DomainShaderRHI,uint32 SamplerIndex,FSamplerStateRHIParamRef NewStateRHI)
 {
@@ -1052,6 +1085,8 @@ void FOpenGLDynamicRHI::RHISetShaderSampler(FDomainShaderRHIParamRef DomainShade
 	{
 		InternalSetSamplerStates(FOpenGL::GetFirstDomainTextureUnit() + SamplerIndex, NewState);
 	}
+	
+	FShaderCache::SetSamplerState(SF_Domain, SamplerIndex, NewStateRHI);
 }
 
 void FOpenGLDynamicRHI::RHISetShaderSampler(FGeometryShaderRHIParamRef GeometryShaderRHI,uint32 SamplerIndex,FSamplerStateRHIParamRef NewStateRHI)
@@ -1075,6 +1110,8 @@ void FOpenGLDynamicRHI::RHISetShaderSampler(FGeometryShaderRHIParamRef GeometryS
 	{
 		InternalSetSamplerStates(FOpenGL::GetFirstGeometryTextureUnit() + SamplerIndex, NewState);
 	}
+	
+	FShaderCache::SetSamplerState(SF_Geometry, SamplerIndex, NewStateRHI);
 }
 
 void FOpenGLDynamicRHI::RHISetShaderTexture(FComputeShaderRHIParamRef ComputeShaderRHI,uint32 TextureIndex,FTextureRHIParamRef NewTextureRHI)
@@ -1090,6 +1127,8 @@ void FOpenGLDynamicRHI::RHISetShaderTexture(FComputeShaderRHIParamRef ComputeSha
 	{
 		InternalSetShaderTexture(nullptr, nullptr, FOpenGL::GetFirstComputeTextureUnit() + TextureIndex, 0, 0, 0, -1);
 	}
+	
+	FShaderCache::SetTexture(SF_Compute, TextureIndex, NewTextureRHI);
 }
 
 void FOpenGLDynamicRHI::RHISetShaderSampler(FPixelShaderRHIParamRef PixelShaderRHI,uint32 SamplerIndex,FSamplerStateRHIParamRef NewStateRHI)
@@ -1113,6 +1152,8 @@ void FOpenGLDynamicRHI::RHISetShaderSampler(FPixelShaderRHIParamRef PixelShaderR
 	{
 		InternalSetSamplerStates(FOpenGL::GetFirstPixelTextureUnit() + SamplerIndex, NewState);
 	}
+	
+	FShaderCache::SetSamplerState(SF_Pixel, SamplerIndex, NewStateRHI);
 }
 
 void FOpenGLDynamicRHI::RHISetShaderUniformBuffer(FVertexShaderRHIParamRef VertexShaderRHI,uint32 BufferIndex,FUniformBufferRHIParamRef BufferRHI)
@@ -1228,6 +1269,8 @@ void FOpenGLDynamicRHI::RHISetDepthStencilState(FDepthStencilStateRHIParamRef Ne
 	DYNAMIC_CAST_OPENGLRESOURCE(DepthStencilState,NewState);
 	PendingState.DepthStencilState = NewState->Data;
 	PendingState.StencilRef = StencilRef;
+	
+	FShaderCache::SetDepthStencilState(NewStateRHI);
 }
 
 void FOpenGLDynamicRHI::UpdateDepthStencilStateInOpenGLContext( FOpenGLContextState& ContextState )
@@ -1586,6 +1629,8 @@ void FOpenGLDynamicRHI::RHISetBlendState(FBlendStateRHIParamRef NewStateRHI,cons
 {
 	DYNAMIC_CAST_OPENGLRESOURCE(BlendState,NewState);
 	FMemory::Memcpy(&PendingState.BlendState,&(NewState->Data),sizeof(FOpenGLBlendStateData));
+	
+	FShaderCache::SetBlendState(NewStateRHI);
 }
 
 void FOpenGLDynamicRHI::RHISetRenderTargets(
@@ -1604,6 +1649,8 @@ void FOpenGLDynamicRHI::RHISetRenderTargets(
 	FMemory::Memset(PendingState.RenderTargetMipmapLevels,0,sizeof(PendingState.RenderTargetMipmapLevels));
 	FMemory::Memset(PendingState.RenderTargetArrayIndex,0,sizeof(PendingState.RenderTargetArrayIndex));
 	PendingState.FirstNonzeroRenderTarget = -1;
+	
+	FShaderCache::SetRenderTargets(NumSimultaneousRenderTargets, NewRenderTargetsRHI, NewDepthStencilTargetRHI);
 
 	for( int32 RenderTargetIndex = NumSimultaneousRenderTargets - 1; RenderTargetIndex >= 0; --RenderTargetIndex )
 	{
@@ -2289,7 +2336,7 @@ FORCEINLINE uint32 GetFirstTextureUnit()
 }
 
 template <EShaderFrequency Frequency>
-FORCEINLINE void SetResource(FOpenGLDynamicRHI* RESTRICT OpenGLRHI, uint32 BindIndex, FOpenGLTextureBase* RESTRICT Texture)
+FORCEINLINE void SetResource(FOpenGLDynamicRHI* RESTRICT OpenGLRHI, uint32 BindIndex, FOpenGLTextureBase* RESTRICT Texture, FRHIResource* Resource)
 {
 	if (Texture)
 	{
@@ -2299,10 +2346,12 @@ FORCEINLINE void SetResource(FOpenGLDynamicRHI* RESTRICT OpenGLRHI, uint32 BindI
 	{
 		OpenGLRHI->InternalSetShaderTexture(Texture, nullptr, GetFirstTextureUnit<Frequency>() + BindIndex, 0, 0, 0, -1);
 	}
+	
+	FShaderCache::SetTexture(Frequency, BindIndex, (FTextureRHIParamRef)Resource);
 }
 
 template <EShaderFrequency Frequency>
-FORCEINLINE void SetResource(FOpenGLDynamicRHI* RESTRICT OpenGLRHI, uint32 BindIndex, FOpenGLSamplerState* RESTRICT SamplerState)
+FORCEINLINE void SetResource(FOpenGLDynamicRHI* RESTRICT OpenGLRHI, uint32 BindIndex, FOpenGLSamplerState* RESTRICT SamplerState, FRHIResource* Resource)
 {
 	if (FOpenGL::SupportsSamplerObjects())
 	{
@@ -2313,13 +2362,17 @@ FORCEINLINE void SetResource(FOpenGLDynamicRHI* RESTRICT OpenGLRHI, uint32 BindI
 	{
 		OpenGLRHI->InternalSetSamplerStates(GetFirstTextureUnit<Frequency>() + BindIndex, SamplerState);
 	}
+	
+	FShaderCache::SetSamplerState(Frequency, BindIndex, (FSamplerStateRHIParamRef)Resource);
 }
 
 template <EShaderFrequency Frequency>
-FORCEINLINE void SetResource(FOpenGLDynamicRHI* RESTRICT OpenGLRHI, uint32 BindIndex, FOpenGLShaderResourceView* RESTRICT SRV)
+FORCEINLINE void SetResource(FOpenGLDynamicRHI* RESTRICT OpenGLRHI, uint32 BindIndex, FOpenGLShaderResourceView* RESTRICT SRV, FRHIResource* Resource)
 {
 	OpenGLRHI->InternalSetShaderTexture(NULL, SRV, GetFirstTextureUnit<Frequency>() + BindIndex, SRV->Target, SRV->Resource, 0, SRV->LimitMip);
-	SetResource<Frequency>(OpenGLRHI,BindIndex,OpenGLRHI->GetPointSamplerState());
+	SetResource<Frequency>(OpenGLRHI,BindIndex,OpenGLRHI->GetPointSamplerState(), OpenGLRHI->GetPointSamplerState());
+	
+	FShaderCache::SetSRV(Frequency, BindIndex, SRV);
 }
 
 template <class GLResourceType, EShaderFrequency ShaderFrequency>
@@ -2338,7 +2391,8 @@ inline int32 SetShaderResourcesFromBuffer(FOpenGLDynamicRHI* RESTRICT OpenGLRHI,
 			const uint8 BindIndex = FRHIResourceTableEntry::GetBindIndex(ResourceInfo);
 
 			GLResourceType* ResourcePtr = (GLResourceType*)Buffer->RawResourceTable[ResourceIndex];
-			SetResource<ShaderFrequency>(OpenGLRHI, BindIndex, ResourcePtr);
+			TRefCountPtr<FRHIResource> Resource = Buffer->ResourceTable[ResourceIndex];
+			SetResource<ShaderFrequency>(OpenGLRHI, BindIndex, ResourcePtr, Resource.GetReference());
 
 			NumSetCalls++;
 			ResourceInfo = *ResourceInfos++;
@@ -2496,6 +2550,8 @@ void FOpenGLDynamicRHI::RHIDrawPrimitive(uint32 PrimitiveType,uint32 BaseVertexI
 		FOpenGL::DrawArraysInstanced(DrawMode, 0, NumElements, NumInstances);
 		REPORT_GL_DRAW_ARRAYS_INSTANCED_EVENT_FOR_FRAME_DUMP( DrawMode, 0, NumElements, NumInstances );
 	}
+
+	FShaderCache::LogDraw(0);
 }
 
 void FOpenGLDynamicRHI::RHIDrawPrimitiveIndirect(uint32 PrimitiveType,FVertexBufferRHIParamRef ArgumentBufferRHI,uint32 ArgumentOffset)
@@ -2541,6 +2597,8 @@ void FOpenGLDynamicRHI::RHIDrawPrimitiveIndirect(uint32 PrimitiveType,FVertexBuf
 			FOpenGL::DrawArraysIndirect( DrawMode, INDEX_TO_VOID(ArgumentOffset));
 		}
 		glBindBuffer( GL_DRAW_INDIRECT_BUFFER, 0);
+		
+		FShaderCache::LogDraw(0);
 	}
 	else
 	{
@@ -2601,6 +2659,8 @@ void FOpenGLDynamicRHI::RHIDrawIndexedIndirect(FIndexBufferRHIParamRef IndexBuff
 			FOpenGL::DrawElementsIndirect( DrawMode, IndexType, INDEX_TO_VOID(DrawArgumentsIndex * 5 *sizeof(uint32)));
 		}
 		glBindBuffer( GL_DRAW_INDIRECT_BUFFER, 0);
+		
+		FShaderCache::LogDraw(IndexBuffer->GetStride());
 	}
 	else
 	{
@@ -2672,6 +2732,8 @@ void FOpenGLDynamicRHI::RHIDrawIndexedPrimitive(FIndexBufferRHIParamRef IndexBuf
 		}
 		REPORT_GL_DRAW_RANGE_ELEMENTS_EVENT_FOR_FRAME_DUMP(DrawMode, MinIndex, MinIndex + NumVertices, NumElements, IndexType, (void *)StartIndex);
 	}
+
+	FShaderCache::LogDraw(IndexBuffer->GetStride());
 }
 
 void FOpenGLDynamicRHI::RHIDrawIndexedPrimitiveIndirect(uint32 PrimitiveType,FIndexBufferRHIParamRef IndexBufferRHI,FVertexBufferRHIParamRef ArgumentBufferRHI,uint32 ArgumentOffset)
@@ -2723,10 +2785,12 @@ void FOpenGLDynamicRHI::RHIDrawIndexedPrimitiveIndirect(uint32 PrimitiveType,FIn
 			FOpenGL::DrawElementsIndirect( DrawMode, IndexType, INDEX_TO_VOID(ArgumentOffset));
 		}
 		glBindBuffer( GL_DRAW_INDIRECT_BUFFER, 0);
+		
+		FShaderCache::LogDraw(IndexBuffer->GetStride());
 	}
 	else
-{
-	UE_LOG(LogRHI, Fatal,TEXT("OpenGL RHI does not yet support indirect draw calls."));
+	{
+		UE_LOG(LogRHI, Fatal,TEXT("OpenGL RHI does not yet support indirect draw calls."));
 	}
 }
 
@@ -2842,6 +2906,8 @@ void FOpenGLDynamicRHI::RHIEndDrawPrimitiveUP()
 	PendingState.NumPrimitives = 0;
 
 	REPORT_GL_DRAW_ARRAYS_EVENT_FOR_FRAME_DUMP( DrawMode, 0, NumElements );
+
+	FShaderCache::LogDraw(0);
 }
 
 /**
@@ -2990,6 +3056,8 @@ void FOpenGLDynamicRHI::RHIEndDrawIndexedPrimitiveUP()
 	PendingState.NumPrimitives = 0;
 
 	REPORT_GL_DRAW_RANGE_ELEMENTS_EVENT_FOR_FRAME_DUMP( DrawMode, PendingState.MinVertexIndex, PendingState.MinVertexIndex + PendingState.NumVertices, NumElements, IndexType, 0 );
+
+	FShaderCache::LogDraw(PendingState.IndexDataStride);
 }
 
 
diff --git a/Engine/Source/Runtime/OpenGLDrv/Private/OpenGLShaders.cpp b/Engine/Source/Runtime/OpenGLDrv/Private/OpenGLShaders.cpp
--- a/Engine/Source/Runtime/OpenGLDrv/Private/OpenGLShaders.cpp
+++ b/Engine/Source/Runtime/OpenGLDrv/Private/OpenGLShaders.cpp
@@ -167,32 +167,11 @@ static bool VerifyCompiledShader(GLuint Shader, const ANSICHAR* GlslCode )
 static bool VerifyProgramPipeline(GLuint Program)
 {
 	bool bOK = true;
-	if (FOpenGL::SupportsSeparateShaderObjects())
+	// Don't try and validate SSOs here - the draw state matters to SSOs and it definitely can't be guaranteed to be valid at this stage
+	if ( FOpenGL::SupportsSeparateShaderObjects() )
 	{
 #if DEBUG_GL_SHADERS
-		FOpenGL::ValidateProgramPipeline(Program);
-		GLint LinkStatus = GL_FALSE;
-		FOpenGL::GetProgramPipelineiv(Program, GL_VALIDATE_STATUS, &LinkStatus);
-		if (LinkStatus == GL_FALSE)
-		{
-			GLint LogLength = 0;
-			FOpenGL::GetProgramPipelineiv(Program, GL_INFO_LOG_LENGTH, &LogLength);
-			ANSICHAR DefaultLog[] = "No log";
-			ANSICHAR *CompileLog = DefaultLog;
-			if (LogLength > 1)
-			{
-				CompileLog = (ANSICHAR *)FMemory::Malloc(LogLength);
-				FOpenGL::GetProgramPipelineInfoLog(Program, LogLength, NULL, CompileLog);
-			}
-			
-			UE_LOG(LogRHI,Warning,TEXT("Failed to validate pipeline %d during compile. This may be harmless. Compile log:\n%s"), Program,
-				   ANSI_TO_TCHAR(CompileLog));
-			
-			if (LogLength > 1)
-			{
-				FMemory::Free(CompileLog);
-			}
-		}
+		bOK = FOpenGL::IsProgramPipeline(Program);
 #endif
 	}
 	else
@@ -373,6 +352,13 @@ namespace
 	}
 }
 
+typedef TArray<ANSICHAR> FAnsiCharArray;
+
+inline uint32 GetTypeHash(FAnsiCharArray const& CharArray)
+{
+	return FCrc::MemCrc32(CharArray.GetData(), CharArray.Num() * sizeof(ANSICHAR));
+}
+
 static void BindShaderLocations(GLenum TypeEnum, GLuint Resource, uint16 InOutMask)
 {
 	if ( OpenGLShaderPlatformNeedsBindLocation(GMaxRHIShaderPlatform) )
@@ -453,8 +439,6 @@ static void BindShaderLocations(GLenum TypeEnum, GLuint Resource, uint16 InOutMa
 template <typename ShaderType>
 ShaderType* CompileOpenGLShader(const TArray<uint8>& Code)
 {
-	typedef TArray<ANSICHAR> FAnsiCharArray;
-
 	SCOPE_CYCLE_COUNTER(STAT_OpenGLShaderCompileTime);
 	VERIFY_GL_SCOPE();
 
@@ -870,9 +854,9 @@ void FOpenGLDynamicRHI::BindUniformBufferBase(FOpenGLContextState& ContextState,
 
 // ============================================================================================================================
 
-static TMap<GLuint, TMap<FString, int64>>& GetOpenGLUniformBlockLocations()
+static TMap<GLuint, TMap<FAnsiCharArray, int64>>& GetOpenGLUniformBlockLocations()
 {
-	static TMap<GLuint, TMap<FString, int64>> UniformBlockLocations;
+	static TMap<GLuint, TMap<FAnsiCharArray, int64>> UniformBlockLocations;
 	return UniformBlockLocations;
 }
 
@@ -884,15 +868,17 @@ static TMap<GLuint, TMap<int64, int64>>& GetOpenGLUniformBlockBindings()
 
 static GLuint GetOpenGLProgramUniformBlockIndex(GLuint Program, const GLchar* UniformBlockName)
 {
-	TMap<FString, int64>& Locations = GetOpenGLUniformBlockLocations().FindOrAdd(Program);
-	int64* Location = Locations.Find(UniformBlockName);
+	TMap<FAnsiCharArray, int64>& Locations = GetOpenGLUniformBlockLocations().FindOrAdd(Program);
+	FAnsiCharArray Name;
+	Name.Append(UniformBlockName, FCStringAnsi::Strlen(UniformBlockName) + 1);
+	int64* Location = Locations.Find(Name);
 	if(Location)
 	{
 		return *Location;
 	}
 	else
 	{
-		int64& Loc = Locations.Emplace(UniformBlockName);
+		int64& Loc = Locations.Emplace(Name);
 		Loc = (int64)FOpenGL::GetUniformBlockIndex(Program, UniformBlockName);
 		return Loc;
 	}
@@ -1729,6 +1715,14 @@ FComputeShaderRHIRef FOpenGLDynamicRHI::RHICreateComputeShader(const TArray<uint
 
 	check( ComputeShader != 0);
 
+	// @todo WARNING: We have to hash here because of the way we immediately link and don't afford the cache a chance to set the OutputHash from ShaderCore.
+	if (FShaderCache::GetShaderCache())
+	{
+		FSHAHash Hash;
+		FSHA1::HashBuffer(Code.GetData(), Code.Num(), Hash.Hash);
+		ComputeShader->SetHash(Hash);
+	}
+
 	FOpenGLLinkedProgramConfiguration Config;
 
 	Config.Shaders[CrossCompiler::SHADER_STAGE_COMPUTE].Resource = ComputeShader->Resource;
@@ -1785,7 +1779,7 @@ static FString GetShaderStageSource(TOpenGLStage* Shader)
 
 struct FOpenGLShaderVaryingMapping
 {
-	FString Name;
+	FAnsiCharArray Name;
 	int32 WriteLoc;
 	int32 ReadLoc;
 };
@@ -1820,25 +1814,33 @@ static void BindShaderStage(FOpenGLLinkedProgramConfiguration::ShaderInfo& Shade
 		else
 		{
 			FOpenGLShaderBindings& PrevStageBindings = PrevStage->Bindings;
-			TMap<FString, int32> PrevStageVaryings;
+			TMap<FAnsiCharArray, int32> PrevStageVaryings;
 			for (int32 i = 0; i < PrevStageBindings.OutputVaryings.Num(); i++)
 			{
-				FString Name = PrevStageBindings.OutputVaryings[i].Varying;
-				Name.RemoveFromStart(TEXT("out_"));
-				Name.RemoveFromStart(TEXT("var_"));
+				FAnsiCharArray Name = PrevStageBindings.OutputVaryings[i].Varying;
+				if ( Name.Num() >= 4 && (FCStringAnsi::Strncmp(Name.GetData(), "out_", 4) == 0 || FCStringAnsi::Strncmp(Name.GetData(), "var_", 4) == 0) )
+				{
+					Name.RemoveAt(0, 4);
+				}
 				PrevStageVaryings.Add(Name, PrevStageBindings.OutputVaryings[i].Location);
 			}
 			
 			bool bInterpolatorMatches = true;
 			
-			TMap<FString, int32> NextStageVaryings;
+			TMap<FAnsiCharArray, int32> NextStageVaryings;
 			TArray<FString> InputErrors;
 			TArray<FOpenGLShaderVaryingMapping> VaryingMapping;
 			for (int32 i = 0; i < NextStageBindings.InputVaryings.Num(); i++)
 			{
-				FString Name = NextStageBindings.InputVaryings[i].Varying;
-				Name.RemoveFromStart(TEXT("in_"));
-				Name.RemoveFromStart(TEXT("var_"));
+				FAnsiCharArray Name = NextStageBindings.InputVaryings[i].Varying;
+				if ( Name.Num() >= 3 && FCStringAnsi::Strncmp(Name.GetData(), "in_", 3) == 0 )
+				{
+					Name.RemoveAt(0, 3);
+				}
+				if ( Name.Num() >= 4 && FCStringAnsi::Strncmp(Name.GetData(), "var_", 4) == 0 )
+				{
+					Name.RemoveAt(0, 4);
+				}
 				NextStageVaryings.Add(Name, NextStageBindings.InputVaryings[i].Location);
 				if( PrevStageVaryings.Contains(Name) )
 				{
@@ -1852,22 +1854,22 @@ static void BindShaderStage(FOpenGLLinkedProgramConfiguration::ShaderInfo& Shade
 							Pair.WriteLoc = PrevLocation;
 							Pair.ReadLoc = NextStageBindings.InputVaryings[i].Location;
 							VaryingMapping.Add(Pair);
-							UE_LOG(LogRHI,Warning,TEXT("Separate Shader Object Binding Warning: Input %s @ %d of stage 0x%x written by stage 0x%x at wrong location %d"), *NextStageBindings.InputVaryings[i].Varying, NextStageBindings.InputVaryings[i].Location, TOpenGLStage0::TypeEnum, TOpenGLStage1::TypeEnum, PrevLocation);
+							UE_LOG(LogRHI,Warning,TEXT("Separate Shader Object Binding Warning: Input %s @ %d of stage 0x%x written by stage 0x%x at wrong location %d"), ANSI_TO_TCHAR(NextStageBindings.InputVaryings[i].Varying.GetData()), NextStageBindings.InputVaryings[i].Location, TOpenGLStage0::TypeEnum, TOpenGLStage1::TypeEnum, PrevLocation);
 						}
 						else if(NextStageBindings.InputVaryings[i].Location == -1)
 						{
-							InputErrors.Add(FString::Printf(TEXT("Separate Shader Object Binding Error: Input %s of stage 0x%x written by stage 0x%x at location %d, can't be rewritten."), *NextStageBindings.InputVaryings[i].Varying, TOpenGLStage0::TypeEnum, TOpenGLStage1::TypeEnum, PrevLocation));
+							InputErrors.Add(FString::Printf(TEXT("Separate Shader Object Binding Error: Input %s of stage 0x%x written by stage 0x%x at location %d, can't be rewritten."), ANSI_TO_TCHAR(NextStageBindings.InputVaryings[i].Varying.GetData()), TOpenGLStage0::TypeEnum, TOpenGLStage1::TypeEnum, PrevLocation));
 											}
 						else
 						{
-							InputErrors.Add(FString::Printf(TEXT("Separate Shader Object Binding Error: Input %s @ %d of stage 0x%x written by stage 0x%x without location, can't be rewritten."), *NextStageBindings.InputVaryings[i].Varying, NextStageBindings.InputVaryings[i].Location, TOpenGLStage0::TypeEnum, TOpenGLStage1::TypeEnum));
+							InputErrors.Add(FString::Printf(TEXT("Separate Shader Object Binding Error: Input %s @ %d of stage 0x%x written by stage 0x%x without location, can't be rewritten."), ANSI_TO_TCHAR(NextStageBindings.InputVaryings[i].Varying.GetData()), NextStageBindings.InputVaryings[i].Location, TOpenGLStage0::TypeEnum, TOpenGLStage1::TypeEnum));
 						}
 						bInterpolatorMatches = false;
 					}
 				}
 				else
 				{
-					InputErrors.Add(FString::Printf(TEXT("Separate Shader Object Binding Error: Input %s @ %d of stage 0x%x not written by stage 0x%x"), *NextStageBindings.InputVaryings[i].Varying, NextStageBindings.InputVaryings[i].Location, TOpenGLStage0::TypeEnum, TOpenGLStage1::TypeEnum));
+					InputErrors.Add(FString::Printf(TEXT("Separate Shader Object Binding Error: Input %s @ %d of stage 0x%x not written by stage 0x%x"), ANSI_TO_TCHAR(NextStageBindings.InputVaryings[i].Varying.GetData()), NextStageBindings.InputVaryings[i].Location, TOpenGLStage0::TypeEnum, TOpenGLStage1::TypeEnum));
 					bInterpolatorMatches = false;
 				}
 			}
@@ -1877,13 +1879,15 @@ static void BindShaderStage(FOpenGLLinkedProgramConfiguration::ShaderInfo& Shade
 			{
 				if ( PrevStageBindings.OutputVaryings[i].Location == -1 )
 				{
-					FString Name = PrevStageBindings.OutputVaryings[i].Varying;
-					Name.RemoveFromStart(TEXT("out_"));
-					Name.RemoveFromStart(TEXT("var_"));
+					FAnsiCharArray Name = PrevStageBindings.OutputVaryings[i].Varying;
+					if ( Name.Num() >= 4 && (FCStringAnsi::Strncmp(Name.GetData(), "out_", 4) == 0 || FCStringAnsi::Strncmp(Name.GetData(), "var_", 4) == 0) )
+					{
+						Name.RemoveAt(0, 4);
+					}
 					if( !NextStageVaryings.Contains(Name) )
 					{
 						OutputElimination.Add(PrevStageBindings.OutputVaryings[i]);
-						UE_LOG(LogRHI,Warning,TEXT("Separate Shader Object Binding Warning: Named output %s of stage 0x%x not read by stage 0x%x"), *PrevStageBindings.OutputVaryings[i].Varying, TOpenGLStage1::TypeEnum, TOpenGLStage0::TypeEnum);
+						UE_LOG(LogRHI,Warning,TEXT("Separate Shader Object Binding Warning: Named output %s of stage 0x%x not read by stage 0x%x"), ANSI_TO_TCHAR(PrevStageBindings.OutputVaryings[i].Varying.GetData()), TOpenGLStage1::TypeEnum, TOpenGLStage0::TypeEnum);
 						bInterpolatorMatches = false;
 					}
 				}
@@ -1931,7 +1935,7 @@ static void BindShaderStage(FOpenGLLinkedProgramConfiguration::ShaderInfo& Shade
 					{
 						for(int32 i = 0; i < PrevLines.Num(); i++)
 						{
-							if(PrevLines[i].Contains(Output.Varying))
+							if(PrevLines[i].Contains(Output.Varying.GetData()))
 							{
 								PrevLines[i].Empty();
 							}
@@ -2217,11 +2221,7 @@ FBoundShaderStateRHIRef FOpenGLDynamicRHI::RHICreateBoundShaderState(
 			DomainShaderRHI
 			);
 
-		FShaderCache* ShaderCache = FShaderCache::GetShaderCache();
-		if(ShaderCache)
-		{
-			ShaderCache->LogBoundShaderState(FOpenGL::GetShaderPlatform(), VertexDeclarationRHI, VertexShaderRHI, PixelShaderRHI, HullShaderRHI, DomainShaderRHI, GeometryShaderRHI);
-		}
+		FShaderCache::LogBoundShaderState(FOpenGL::GetShaderPlatform(), VertexDeclarationRHI, VertexShaderRHI, PixelShaderRHI, HullShaderRHI, DomainShaderRHI, GeometryShaderRHI, BoundShaderState);
 
 		return BoundShaderState;
 	}
diff --git a/Engine/Source/Runtime/OpenGLDrv/Private/OpenGLState.cpp b/Engine/Source/Runtime/OpenGLDrv/Private/OpenGLState.cpp
--- a/Engine/Source/Runtime/OpenGLDrv/Private/OpenGLState.cpp
+++ b/Engine/Source/Runtime/OpenGLDrv/Private/OpenGLState.cpp
@@ -5,6 +5,7 @@
 =============================================================================*/
 
 #include "OpenGLDrvPrivate.h"
+#include "ShaderCache.h"
 
 GLint GMaxOpenGLTextureFilterAnisotropic = 1;
 
@@ -258,6 +259,8 @@ FSamplerStateRHIRef FOpenGLDynamicRHI::RHICreateSamplerState(const FSamplerState
 	// Manually add reference as we control the creation/destructions
 	SamplerState->AddRef();
 	GSamplerStateCache.Add(Initializer, SamplerState);
+
+	FShaderCache::LogSamplerState(Initializer, SamplerState);
 	return SamplerState;
 }
 
@@ -268,6 +271,8 @@ FRasterizerStateRHIRef FOpenGLDynamicRHI::RHICreateRasterizerState(const FRaster
 	RasterizerState->Data.FillMode = TranslateFillMode(Initializer.FillMode);
 	RasterizerState->Data.DepthBias = Initializer.DepthBias;
 	RasterizerState->Data.SlopeScaleDepthBias = Initializer.SlopeScaleDepthBias;
+	
+	FShaderCache::LogRasterizerState(Initializer, RasterizerState);
 	return RasterizerState;
 }
 
@@ -289,6 +294,8 @@ FDepthStencilStateRHIRef FOpenGLDynamicRHI::RHICreateDepthStencilState(const FDe
 	DepthStencilState->Data.CCWStencilPass = TranslateStencilOp(Initializer.BackFacePassStencilOp);
 	DepthStencilState->Data.StencilReadMask = Initializer.StencilReadMask;
 	DepthStencilState->Data.StencilWriteMask = Initializer.StencilWriteMask;
+
+	FShaderCache::LogDepthStencilState(Initializer, DepthStencilState);
 	return DepthStencilState;
 }
 
@@ -319,5 +326,7 @@ FBlendStateRHIRef FOpenGLDynamicRHI::RHICreateBlendState(const FBlendStateInitia
 		RenderTarget.ColorWriteMaskB = (RenderTargetInitializer.ColorWriteMask & CW_BLUE) != 0;
 		RenderTarget.ColorWriteMaskA = (RenderTargetInitializer.ColorWriteMask & CW_ALPHA) != 0;
 	}
+	
+	FShaderCache::LogBlendState(Initializer, BlendState);
 	return BlendState;
 }
diff --git a/Engine/Source/Runtime/OpenGLDrv/Private/OpenGLTexture.cpp b/Engine/Source/Runtime/OpenGLDrv/Private/OpenGLTexture.cpp
--- a/Engine/Source/Runtime/OpenGLDrv/Private/OpenGLTexture.cpp
+++ b/Engine/Source/Runtime/OpenGLDrv/Private/OpenGLTexture.cpp
@@ -5,6 +5,7 @@
 =============================================================================*/
 
 #include "OpenGLDrvPrivate.h"
+#include "ShaderCache.h"
 
 /*-----------------------------------------------------------------------------
 	Texture allocator support.
@@ -108,6 +109,8 @@ void OpenGLTextureAllocated(FRHITexture* Texture, uint32 Flags)
 
 void OpenGLTextureDeleted( FRHITexture* Texture )
 {
+	FShaderCache::RemoveTexture(Texture);
+	
 	bool bRenderTarget = !ShouldCountAsTextureMemory(Texture->GetFlags());
 	int32 TextureSize = 0;
 	if (Texture->GetTextureCube())
@@ -1689,6 +1692,8 @@ FShaderResourceViewRHIRef FOpenGLDynamicRHI::RHICreateShaderResourceView(FTextur
 	{
 		View = new FOpenGLShaderResourceView(this, Texture2D->Resource, Texture2D->Target, MipLevel, false);
 	}
+	
+	FShaderCache::LogSRV(View, Texture2DRHI, MipLevel, Texture2DRHI->GetNumMips(), Texture2DRHI->GetFormat());
 
 	return View;
 }
@@ -1742,6 +1747,8 @@ FShaderResourceViewRHIRef FOpenGLDynamicRHI::RHICreateShaderResourceView(FTextur
 	{
 		View = new FOpenGLShaderResourceView(this, Texture2D->Resource, Texture2D->Target, MipLevel, false);
 	}
+	
+	FShaderCache::LogSRV(View, Texture2DRHI, MipLevel, NumMipLevels, Format);
 
 	return View;
 }
diff --git a/Engine/Source/Runtime/OpenGLDrv/Private/OpenGLUAV.cpp b/Engine/Source/Runtime/OpenGLDrv/Private/OpenGLUAV.cpp
--- a/Engine/Source/Runtime/OpenGLDrv/Private/OpenGLUAV.cpp
+++ b/Engine/Source/Runtime/OpenGLDrv/Private/OpenGLUAV.cpp
@@ -2,6 +2,7 @@
 
 
 #include "OpenGLDrvPrivate.h"
+#include "ShaderCache.h"
 
 FShaderResourceViewRHIRef FOpenGLDynamicRHI::RHICreateShaderResourceView(FVertexBufferRHIParamRef VertexBufferRHI, uint32 Stride, uint8 Format)
 {
@@ -31,11 +32,16 @@ FShaderResourceViewRHIRef FOpenGLDynamicRHI::RHICreateShaderResourceView(FVertex
 	// and the next draw will take care of cleaning it up; or
 	// next operation that needs the stage will switch something else in on it.
 
-	return new FOpenGLShaderResourceView(this,TextureID,GL_TEXTURE_BUFFER,VertexBufferRHI,Format);
+	FShaderResourceViewRHIRef Result = new FOpenGLShaderResourceView(this,TextureID,GL_TEXTURE_BUFFER,VertexBufferRHI,Format);
+	FShaderCache::LogSRV(Result, VertexBufferRHI, Stride, Format);
+
+	return Result;
 }
 
 FOpenGLShaderResourceView::~FOpenGLShaderResourceView()
 {
+	FShaderCache::RemoveSRV(this);
+	
 	if (Resource && OwnsResource)
 	{
 		OpenGLRHI->InvalidateTextureResourceInCache( Resource );
diff --git a/Engine/Source/Runtime/OpenGLDrv/Private/OpenGLVertexDeclaration.cpp b/Engine/Source/Runtime/OpenGLDrv/Private/OpenGLVertexDeclaration.cpp
--- a/Engine/Source/Runtime/OpenGLDrv/Private/OpenGLVertexDeclaration.cpp
+++ b/Engine/Source/Runtime/OpenGLDrv/Private/OpenGLVertexDeclaration.cpp
@@ -135,13 +135,9 @@ FVertexDeclarationRHIRef FOpenGLDynamicRHI::RHICreateVertexDeclaration(const FVe
 		// Create and add to the cache if it doesn't exist.
 		VertexDeclarationRefPtr = &GOpenGLVertexDeclarationCache.Add(Key,new FOpenGLVertexDeclaration(Key.VertexElements));
 		
-		FShaderCache* ShaderCache = FShaderCache::GetShaderCache();
-		if(ShaderCache)
-		{
-			check(VertexDeclarationRefPtr);
-			check(IsValidRef(*VertexDeclarationRefPtr));
-			ShaderCache->LogVertexDeclaration(Elements, *VertexDeclarationRefPtr);
-		}
+		check(VertexDeclarationRefPtr);
+		check(IsValidRef(*VertexDeclarationRefPtr));
+		FShaderCache::LogVertexDeclaration(Elements, *VertexDeclarationRefPtr);
 	}
 
 	// The cached declaration must match the input declaration!
diff --git a/Engine/Source/Runtime/OpenGLDrv/Private/Windows/OpenGLWindows.h b/Engine/Source/Runtime/OpenGLDrv/Private/Windows/OpenGLWindows.h
--- a/Engine/Source/Runtime/OpenGLDrv/Private/Windows/OpenGLWindows.h
+++ b/Engine/Source/Runtime/OpenGLDrv/Private/Windows/OpenGLWindows.h
@@ -380,7 +380,8 @@
 	EnumMacro(PFNGLPROGRAMUNIFORM4UIVPROC, glProgramUniform4uiv) \
 	EnumMacro(PFNGLGETPROGRAMPIPELINEIVPROC, glGetProgramPipelineiv) \
 	EnumMacro(PFNGLVALIDATEPROGRAMPIPELINEPROC, glValidateProgramPipeline) \
-	EnumMacro(PFNGLGETPROGRAMPIPELINEINFOLOGPROC, glGetProgramPipelineInfoLog)
+	EnumMacro(PFNGLGETPROGRAMPIPELINEINFOLOGPROC, glGetProgramPipelineInfoLog) \
+	EnumMacro(PFNGLISPROGRAMPIPELINEPROC, glIsProgramPipeline)
 
 /** List of all OpenGL entry points. */
 #define ENUM_GL_ENTRYPOINTS_ALL(EnumMacro) \
diff --git a/Engine/Source/Runtime/OpenGLDrv/Public/OpenGL.h b/Engine/Source/Runtime/OpenGLDrv/Public/OpenGL.h
--- a/Engine/Source/Runtime/OpenGLDrv/Public/OpenGL.h
+++ b/Engine/Source/Runtime/OpenGLDrv/Public/OpenGL.h
@@ -316,6 +316,7 @@ public:
 	static FORCEINLINE void GetProgramPipelineiv(GLuint Pipeline, GLenum Pname, GLint* Params) UGL_OPTIONAL_VOID
 	static FORCEINLINE void ValidateProgramPipeline(GLuint Pipeline) UGL_OPTIONAL_VOID
 	static FORCEINLINE void GetProgramPipelineInfoLog(GLuint Pipeline, GLsizei BufSize, GLsizei* Length, GLchar* InfoLog) UGL_OPTIONAL_VOID
+	static FORCEINLINE bool IsProgramPipeline(GLuint Pipeline) UGL_OPTIONAL(false)
 
 	static FORCEINLINE GLuint64 GetTextureSamplerHandle(GLuint Texture, GLuint Sampler) UGL_REQUIRED(0)
 	static FORCEINLINE GLuint64 GetTextureHandle(GLuint Texture) UGL_REQUIRED(0)
diff --git a/Engine/Source/Runtime/OpenGLDrv/Public/OpenGL3.h b/Engine/Source/Runtime/OpenGLDrv/Public/OpenGL3.h
--- a/Engine/Source/Runtime/OpenGLDrv/Public/OpenGL3.h
+++ b/Engine/Source/Runtime/OpenGLDrv/Public/OpenGL3.h
@@ -712,6 +712,10 @@ struct FOpenGL3 : public FOpenGLBase
 	{
 		glGetProgramPipelineInfoLog(Pipeline, BufSize, Length, InfoLog);
 	}
+	static FORCEINLINE bool IsProgramPipeline(GLuint Pipeline)
+	{
+		return (glIsProgramPipeline(Pipeline) == GL_TRUE);
+	}
 
 	static FORCEINLINE ERHIFeatureLevel::Type GetFeatureLevel()
 	{
diff --git a/Engine/Source/Runtime/OpenGLDrv/Public/OpenGLResources.h b/Engine/Source/Runtime/OpenGLDrv/Public/OpenGLResources.h
--- a/Engine/Source/Runtime/OpenGLDrv/Public/OpenGLResources.h
+++ b/Engine/Source/Runtime/OpenGLDrv/Public/OpenGLResources.h
@@ -8,6 +8,7 @@
 
 #include "BoundShaderStateCache.h"
 #include "OpenGLShaderResources.h"
+#include "ShaderCache.h"
 
 extern void OnVertexBufferDeletion( GLuint VertexBufferResource );
 extern void OnIndexBufferDeletion( GLuint IndexBufferResource );
@@ -863,6 +864,65 @@ public:
 		PixelBuffers.AddZeroed(this->GetNumMips() * (bCubemap ? 6 : 1) * GetEffectiveSizeZ());
 		bAllocatedStorage.Init(bInAllocatedStorage, this->GetNumMips() * (bCubemap ? 6 : 1));
 		ClientStorageBuffers.AddZeroed(this->GetNumMips() * (bCubemap ? 6 : 1) * GetEffectiveSizeZ());
+	
+		FShaderCache* ShaderCache = FShaderCache::GetShaderCache();
+		if ( ShaderCache )
+		{
+			FShaderTextureKey Tex;
+			Tex.Format = (EPixelFormat)InFormat;
+			Tex.Flags = InFlags;
+			Tex.MipLevels = InNumMips;
+			Tex.Samples = InNumSamples;
+			Tex.X = InSizeX;
+			Tex.Y = InSizeY;
+			Tex.Z = InSizeZ;
+			switch(InTarget)
+			{
+				case GL_TEXTURE_1D:
+				{
+					Tex.Type = SCTT_Texture1D;
+					break;
+				}
+				case GL_TEXTURE_2D:
+				case GL_TEXTURE_2D_MULTISAMPLE:
+				{
+					Tex.Type = SCTT_Texture2D;
+					break;
+				}
+				case GL_TEXTURE_3D:
+				{
+					Tex.Type = SCTT_Texture3D;
+					break;
+				}
+				case GL_TEXTURE_CUBE_MAP:
+				{
+					Tex.Type = SCTT_TextureCube;
+					break;
+				}
+				case GL_TEXTURE_1D_ARRAY:
+				{
+					Tex.Type = SCTT_Texture1DArray;
+					break;
+				}
+				case GL_TEXTURE_2D_ARRAY:
+				{
+					Tex.Type = SCTT_Texture2DArray;
+					break;
+				}
+				case GL_TEXTURE_CUBE_MAP_ARRAY:
+				{
+					Tex.Type = SCTT_TextureCubeArray;
+					Tex.Z = InArraySize;
+					break;
+				}
+				default:
+				{
+					check(false);
+				}
+			}
+		
+			FShaderCache::LogTexture(Tex, this);
+		}
 	}
 
 	virtual ~TOpenGLTexture();
diff --git a/Engine/Source/Runtime/OpenGLDrv/Public/OpenGLShaderResources.h b/Engine/Source/Runtime/OpenGLDrv/Public/OpenGLShaderResources.h
--- a/Engine/Source/Runtime/OpenGLDrv/Public/OpenGLShaderResources.h
+++ b/Engine/Source/Runtime/OpenGLDrv/Public/OpenGLShaderResources.h
@@ -50,14 +50,14 @@ inline FArchive& operator<<(FArchive& Ar, FOpenGLShaderResourceTable& SRT)
 
 struct FOpenGLShaderVarying
 {
-	FString Varying;
+	TArray<ANSICHAR> Varying;
 	int32 Location;
 	
 	friend bool operator==(const FOpenGLShaderVarying& A, const FOpenGLShaderVarying& B)
 	{
 		if(&A != &B)
 		{
-			return (A.Location == B.Location) && (A.Varying == B.Varying);
+			return (A.Location == B.Location) && (A.Varying.Num() == B.Varying.Num()) && (FMemory::Memcmp(A.Varying.GetData(), B.Varying.GetData(), A.Varying.Num() * sizeof(ANSICHAR)) == 0);
 		}
 		return true;
 	}
@@ -65,7 +65,7 @@ struct FOpenGLShaderVarying
 	friend uint32 GetTypeHash(const FOpenGLShaderVarying &Var)
 	{
 		uint32 Hash = GetTypeHash(Var.Location);
-		Hash ^= GetTypeHash(Var.Varying);
+		Hash ^= FCrc::MemCrc32(Var.Varying.GetData(), Var.Varying.Num() * sizeof(ANSICHAR));
 		return Hash;
 	}
 };
diff --git a/Engine/Source/Runtime/ShaderCore/Private/Shader.cpp b/Engine/Source/Runtime/ShaderCore/Private/Shader.cpp
--- a/Engine/Source/Runtime/ShaderCore/Private/Shader.cpp
+++ b/Engine/Source/Runtime/ShaderCore/Private/Shader.cpp
@@ -345,11 +345,7 @@ void FShaderResource::Serialize(FArchive& Ar)
 		INC_DWORD_STAT_BY_FName(GetMemoryStatType((EShaderFrequency)Target.Frequency).GetName(), (int64)Code.Num());
 		INC_DWORD_STAT_BY(STAT_Shaders_ShaderResourceMemory, GetSizeBytes());
 		
-		FShaderCache* ShaderCache = FShaderCache::GetShaderCache();
-		if(ShaderCache)
-		{
-			ShaderCache->LogShader((EShaderPlatform)Target.Platform, (EShaderFrequency)Target.Frequency, OutputHash, Code);
-		}
+		FShaderCache::LogShader((EShaderPlatform)Target.Platform, (EShaderFrequency)Target.Frequency, OutputHash, Code);
 	}
 }
 
@@ -493,7 +489,7 @@ void FShaderResource::InitRHI()
 	}
 	else if(Target.Frequency == SF_Compute)
 	{
-		ComputeShader = ShaderCache ? ShaderCache->GetComputeShader((EShaderPlatform)Target.Platform, OutputHash, Code) : RHICreateComputeShader(Code);
+		ComputeShader = ShaderCache ? ShaderCache->GetComputeShader((EShaderPlatform)Target.Platform, Code) : RHICreateComputeShader(Code);
 	}
 
 	if (Target.Frequency != SF_Geometry)
diff --git a/Engine/Source/Runtime/ShaderCore/Private/ShaderCache.cpp b/Engine/Source/Runtime/ShaderCore/Private/ShaderCache.cpp
--- a/Engine/Source/Runtime/ShaderCore/Private/ShaderCache.cpp
+++ b/Engine/Source/Runtime/ShaderCore/Private/ShaderCache.cpp
@@ -9,14 +9,85 @@
 #include "RHI.h"
 #include "RenderingThread.h"
 
-int32 bUseShaderCaching = PLATFORM_MAC ? 1 : 0;
-static FAutoConsoleVariableRef CVarUseShaderCaching(
+const FGuid FShaderCacheCustomVersion::Key(0xB954F018, 0xC9624DD6, 0xA74E79B1, 0x8EA113C2);
+const FGuid FShaderCacheCustomVersion::GameKey(0x03D4EB48, 0xB50B4CC3, 0xA598DE41, 0x5C6CC993);
+FCustomVersionRegistration GRegisterShaderCacheVersion(FShaderCacheCustomVersion::Key, FShaderCacheCustomVersion::Latest, TEXT("ShaderCacheVersion"));
+static TCHAR const* GShaderCacheFileName = TEXT("ShaderCache.ushadercache");
+
+// Only the cooked Mac build defaults to using the shader cache for now, Editor is too likely to invalidate shader keys leading to ever growing cache
+int32 FShaderCache::bUseShaderCaching = (PLATFORM_MAC && !WITH_EDITOR) ? 1 : 0;
+FAutoConsoleVariableRef FShaderCache::CVarUseShaderCaching(
 	TEXT("r.UseShaderCaching"),
 	bUseShaderCaching,
-	TEXT("If true, log all shaders & bound-shader-states, so that they may be instantiated in the RHI on deserialisation rather than waiting for first use.")
+	TEXT("If true, log all shaders & bound-shader-states, so that they may be instantiated in the RHI on deserialisation rather than waiting for first use."),
+	ECVF_ReadOnly|ECVF_RenderThreadSafe
+	);
+
+// Predrawing takes an existing shader cache with draw log & renders each shader + draw-state combination before use to avoid in-driver recompilation
+// This requires plenty of setup & is done in batches at frame-end.
+int32 FShaderCache::bUseShaderPredraw = (PLATFORM_MAC && !WITH_EDITOR) ? 1 : 0;
+FAutoConsoleVariableRef FShaderCache::CVarUseShaderPredraw(
+	TEXT("r.UseShaderPredraw"),
+	bUseShaderPredraw,
+	TEXT("Use an existing draw-log to predraw shaders in batches before being used to reduce hitches due to in-driver recompilation."),
+	ECVF_ReadOnly|ECVF_RenderThreadSafe
+	);
+
+// The actual draw loggging is even more expensive as it has to cache all the RHI draw state & is disabled by default if you aren't a developer.
+int32 FShaderCache::bUseShaderDrawLog = (PLATFORM_MAC && !WITH_EDITOR && !UE_BUILD_SHIPPING) ? 1 : 0;
+FAutoConsoleVariableRef FShaderCache::CVarUseShaderDrawLog(
+	TEXT("r.UseShaderDrawLog"),
+	bUseShaderDrawLog,
+	TEXT("If true, log all the draw states used for each shader pipeline, so that they may be pre-drawn in batches (see: r.UseShaderPredraw). This can be expensive & should be used only when generating the shader cache."),
+	ECVF_ReadOnly|ECVF_RenderThreadSafe
+	);
+
+// As predrawing can take significant time then batch the draws up into chunks defined by frame time
+int32 FShaderCache::PredrawBatchTime = -1;
+FAutoConsoleVariableRef FShaderCache::CVarPredrawBatchTime(
+	TEXT("r.PredrawBatchTime"),
+	PredrawBatchTime,
+	TEXT("Time in ms to spend predrawing shaders each frame, or -1 to perform all predraws immediately."),
+	ECVF_RenderThreadSafe
 	);
 
 FShaderCache* FShaderCache::Cache = nullptr;
+int32 FShaderCache::GameVersion = 0;
+
+static bool ShaderPlatformCanPrebindBoundShaderState(EShaderPlatform Platform)
+{
+	switch (Platform)
+	{
+		case SP_PCD3D_SM5:
+		case SP_PS4:
+		case SP_XBOXONE:
+		case SP_PCD3D_SM4:
+		case SP_PCD3D_ES2:
+		case SP_METAL:
+		case SP_OPENGL_SM4_MAC:
+		case SP_METAL_MRT:
+		{
+			return true;
+		}
+		case SP_OPENGL_SM4:
+		case SP_OPENGL_PCES2:
+		case SP_OPENGL_SM5:
+		case SP_OPENGL_ES2:
+		case SP_OPENGL_ES2_WEBGL:
+		case SP_OPENGL_ES2_IOS:
+		case SP_OPENGL_ES31_EXT:
+		default:
+		{
+			return false;
+		}
+	}
+}
+
+void FShaderCache::SetGameVersion(int32 InGameVersion)
+{
+	check(!Cache);
+	GameVersion = InGameVersion;
+}
 
 void FShaderCache::InitShaderCache()
 {
@@ -27,11 +98,6 @@ void FShaderCache::InitShaderCache()
 	}
 }
 
-FShaderCache* FShaderCache::GetShaderCache()
-{
-	return bUseShaderCaching ? Cache : nullptr;
-}
-
 void FShaderCache::ShutdownShaderCache()
 {
 	if (Cache)
@@ -41,31 +107,71 @@ void FShaderCache::ShutdownShaderCache()
 	}
 }
 
-FShaderCache::FShaderCache()
+FArchive& operator<<( FArchive& Ar, FShaderCache& Info )
 {
-	FString UserBinaryShaderFile = FPaths::GameSavedDir() / TEXT("ShaderCache.ushaderprecache");
-	FString GameBinaryShaderFile = FPaths::GameContentDir() / TEXT("ShaderCache.ushaderprecache");
-
-	FArchive* BinaryShaderAr = nullptr;
-	if ( IFileManager::Get().FileSize(*UserBinaryShaderFile) > 0 )
+	Ar.UsingCustomVersion(FShaderCacheCustomVersion::Key);
+	Ar.SetCustomVersion(FShaderCacheCustomVersion::GameKey, FShaderCache::GameVersion, TEXT("ShaderCacheGameVersion"));
+	
+	FCustomVersionContainer Vers = Ar.GetCustomVersions();
+	Vers.Serialize(Ar);
+	if ( Ar.IsLoading() )
 	{
-		BinaryShaderAr = IFileManager::Get().CreateFileReader(*UserBinaryShaderFile);
+		Ar.SetCustomVersions(Vers);
 	}
-	else
+	
+	if ( !Ar.IsError() && Ar.CustomVer(FShaderCacheCustomVersion::Key) == FShaderCacheCustomVersion::Latest && Ar.CustomVer(FShaderCacheCustomVersion::GameKey) == FShaderCache::GameVersion )
 	{
-		BinaryShaderAr = IFileManager::Get().CreateFileReader(*GameBinaryShaderFile);
+		Ar << Info.Caches;
 	}
+	return Ar;
+}
+
+FShaderCache::FShaderCache()
+: bCurrentDepthStencilTarget(false)
+, CurrentNumRenderTargets(0)
+, CurrentShaderState(nullptr)
+, bIsPreDraw(false)
+{
+	Viewport[0] = Viewport[1] = Viewport[2] = Viewport[3] = 0;
+	DepthRange[0] = DepthRange[1] = 0.0f;
+
+	FString UserBinaryShaderFile = FPaths::GameSavedDir() / GShaderCacheFileName;
+	FString GameBinaryShaderFile = FPaths::GameContentDir() / GShaderCacheFileName;
 
-	if ( BinaryShaderAr != nullptr )
+	// Try to load user cache, making sure that if we fail version test we still try game-content version.
+	bool bLoadedUserCache = false;
+	if ( IFileManager::Get().FileSize(*UserBinaryShaderFile) > 0 )
 	{
-		*BinaryShaderAr << *this;
-		delete BinaryShaderAr;
+		FArchive* BinaryShaderAr = IFileManager::Get().CreateFileReader(*UserBinaryShaderFile);
+
+		if ( BinaryShaderAr != nullptr )
+		{
+			*BinaryShaderAr << *this;
+			
+			if ( !BinaryShaderAr->IsError() && BinaryShaderAr->CustomVer(FShaderCacheCustomVersion::Key) == FShaderCacheCustomVersion::Latest && BinaryShaderAr->CustomVer(FShaderCacheCustomVersion::GameKey) == FShaderCache::GameVersion )
+			{
+				bLoadedUserCache = true;
+			}
+			
+			delete BinaryShaderAr;
+		}
+	}
+	
+	// Fallback to game-content version.
+	if ( !bLoadedUserCache && IFileManager::Get().FileSize(*GameBinaryShaderFile) > 0 )
+	{
+		FArchive* BinaryShaderAr = IFileManager::Get().CreateFileReader(*GameBinaryShaderFile);
+		if ( BinaryShaderAr != nullptr )
+		{
+			*BinaryShaderAr << *this;
+			delete BinaryShaderAr;
+		}
 	}
 }
 
 FShaderCache::~FShaderCache()
 {
-	FString BinaryShaderFile = FPaths::GameSavedDir() / TEXT("ShaderCache.ushaderprecache");
+	FString BinaryShaderFile = FPaths::GameSavedDir() / GShaderCacheFileName;
 	FArchive* BinaryShaderAr = IFileManager::Get().CreateFileWriter(*BinaryShaderFile);
 	if( BinaryShaderAr != NULL )
 	{
@@ -179,12 +285,13 @@ FDomainShaderRHIRef FShaderCache::GetDomainShader(EShaderPlatform Platform, FSHA
 	return Shader;
 }
 
-FComputeShaderRHIRef FShaderCache::GetComputeShader(EShaderPlatform Platform, FSHAHash Hash, TArray<uint8> const& Code)
+FComputeShaderRHIRef FShaderCache::GetComputeShader(EShaderPlatform Platform, TArray<uint8> const& Code)
 {
 	FShaderCacheKey Key;
 	Key.Platform = Platform;
 	Key.Frequency = SF_Compute;
-	Key.Hash = Hash;
+	// @todo WARNING: RHI is responsible for hashing Compute shaders due to the way OpenGLDrv implements compute!
+	FSHA1::HashBuffer(Code.GetData(), Code.Num(), Key.Hash.Hash);
 	Key.bActive = true;
 	FComputeShaderRHIRef Shader = CachedComputeShaders.FindRef(Key);
 	if(!IsValidRef(Shader))
@@ -193,14 +300,13 @@ FComputeShaderRHIRef FShaderCache::GetComputeShader(EShaderPlatform Platform, FS
 		PlatformCache.Shaders.Add(Key);
 		Shader = RHICreateComputeShader(Code);
 		check(IsValidRef(Shader));
-		Shader->SetHash(Hash);
 		CachedComputeShaders.Add(Key, Shader);
 		PrebindShader(Key);
 	}
 	return Shader;
 }
 
-void FShaderCache::LogShader(EShaderPlatform Platform, EShaderFrequency Frequency, FSHAHash Hash, TArray<uint8> const& Code)
+void FShaderCache::InternalLogShader(EShaderPlatform Platform, EShaderFrequency Frequency, FSHAHash Hash, TArray<uint8> const& Code)
 {
 	FShaderCacheKey Key;
 	Key.Hash = Hash;
@@ -212,17 +318,18 @@ void FShaderCache::LogShader(EShaderPlatform Platform, EShaderFrequency Frequenc
 	});
 }
 
-void FShaderCache::LogVertexDeclaration(const FVertexDeclarationElementList& VertexElements, FVertexDeclarationRHIParamRef VertexDeclaration)
+void FShaderCache::InternalLogVertexDeclaration(const FVertexDeclarationElementList& VertexElements, FVertexDeclarationRHIParamRef VertexDeclaration)
 {
 	VertexDeclarations.Add(VertexDeclaration, VertexElements);
 }
 
-void FShaderCache::LogBoundShaderState(EShaderPlatform Platform, FVertexDeclarationRHIParamRef VertexDeclaration,
+void FShaderCache::InternalLogBoundShaderState(EShaderPlatform Platform, FVertexDeclarationRHIParamRef VertexDeclaration,
 								FVertexShaderRHIParamRef VertexShader,
 								FPixelShaderRHIParamRef PixelShader,
 								FHullShaderRHIParamRef HullShader,
 								FDomainShaderRHIParamRef DomainShader,
-								FGeometryShaderRHIParamRef GeometryShader)
+								FGeometryShaderRHIParamRef GeometryShader,
+								FBoundShaderStateRHIParamRef BoundState)
 {
 	FShaderCacheBoundState Info;
 	if(VertexDeclaration)
@@ -308,52 +415,431 @@ void FShaderCache::LogBoundShaderState(EShaderPlatform Platform, FVertexDeclarat
 			Set.Add(Info);
 		}
 	}
+	if ( bUseShaderPredraw || bUseShaderDrawLog )
+	{
+		ShaderStates.Add(BoundState, Info);
+	}
+}
+	
+void FShaderCache::InternalLogBlendState(FBlendStateInitializerRHI const& Init, FBlendStateRHIParamRef State)
+{
+	if ( bUseShaderPredraw || bUseShaderDrawLog )
+	{
+		BlendStates.Add(State, Init);
+	}
 }
 
-void FShaderCache::PrebindShader(FShaderCacheKey Key)
+void FShaderCache::InternalLogRasterizerState(FRasterizerStateInitializerRHI const& Init, FRasterizerStateRHIParamRef State)
 {
-	FShaderPlatformCache& PlatformCache = Caches.FindOrAdd(Key.Platform);
-	TSet<FShaderCacheBoundState>& BoundStates = PlatformCache.ShaderStateMembership.FindOrAdd(Key);
-	for (FShaderCacheBoundState State : BoundStates)
-	{
-		FVertexShaderRHIRef VertexShader = State.VertexShader.bActive ? CachedVertexShaders.FindRef(State.VertexShader) : nullptr;
-		FPixelShaderRHIRef PixelShader = State.PixelShader.bActive ? CachedPixelShaders.FindRef(State.PixelShader) : nullptr;
-		FGeometryShaderRHIRef GeometryShader = State.GeometryShader.bActive ? CachedGeometryShaders.FindRef(State.GeometryShader) : nullptr;
-		FHullShaderRHIRef HullShader = State.HullShader.bActive ? CachedHullShaders.FindRef(State.HullShader) : nullptr;
-		FDomainShaderRHIRef DomainShader = State.DomainShader.bActive ? CachedDomainShaders.FindRef(State.DomainShader) : nullptr;
+	if ( bUseShaderPredraw || bUseShaderDrawLog )
+	{
+		RasterizerStates.Add(State, Init);
+	}
+}
+
+void FShaderCache::InternalLogDepthStencilState(FDepthStencilStateInitializerRHI const& Init, FDepthStencilStateRHIParamRef State)
+{
+	if ( bUseShaderPredraw || bUseShaderDrawLog )
+	{
+		DepthStencilStates.Add(State, Init);
+	}
+}
+
+void FShaderCache::InternalLogSamplerState(FSamplerStateInitializerRHI const& Init, FSamplerStateRHIParamRef State)
+{
+	if ( bUseShaderDrawLog )
+	{
+		SamplerStates.Add(State, Init);
 		
-		bool bOK = true;
-		bOK &= (State.VertexShader.bActive == IsValidRef(VertexShader));
-		bOK &= (State.PixelShader.bActive == IsValidRef(PixelShader));
-		bOK &= (State.GeometryShader.bActive == IsValidRef(GeometryShader));
-		bOK &= (State.HullShader.bActive == IsValidRef(HullShader));
-		bOK &= (State.DomainShader.bActive == IsValidRef(DomainShader));
+		FShaderPlatformCache& PlatformCache = Caches.FindOrAdd(GMaxRHIShaderPlatform);
+		PlatformCache.SamplerStates.Add(Init);
+	}
+}
+
+void FShaderCache::InternalLogTexture(FShaderTextureKey const& Init, FTextureRHIParamRef State)
+{
+	if ( bUseShaderPredraw || bUseShaderDrawLog )
+	{
+		Textures.Add(State, Init);
+		CachedTextures.Add(Init, State);
+		
+		FShaderPlatformCache& PlatformCache = Caches.FindOrAdd(GMaxRHIShaderPlatform);
+		FShaderResourceKey Key;
+		Key.Tex = Init;
+		Key.Format = Init.Format;
+		PlatformCache.Resources.Add(Key);
+	}
+}
+
+void FShaderCache::InternalLogSRV(FShaderResourceViewRHIParamRef SRV, FTextureRHIParamRef Texture, uint8 StartMip, uint8 NumMips, uint8 Format)
+{
+	if ( bUseShaderPredraw || bUseShaderDrawLog )
+	{
+		FShaderResourceKey Key;
+		Key.Tex = Textures.FindChecked(Texture);
+		Key.BaseMip = StartMip;
+		Key.MipLevels = NumMips;
+		Key.Format = Format;
+		Key.bSRV = true;
+		SRVs.Add(SRV, Key);
+		CachedSRVs.Add(Key, FShaderResourceViewBinding(SRV, nullptr, Texture));
+		
+		FShaderPlatformCache& PlatformCache = Caches.FindOrAdd(GMaxRHIShaderPlatform);
+		PlatformCache.Resources.Add(Key);
+	}
+}
+
+void FShaderCache::InternalLogSRV(FShaderResourceViewRHIParamRef SRV, FVertexBufferRHIParamRef Vb, uint32 Stride, uint8 Format)
+{
+	if ( bUseShaderPredraw || bUseShaderDrawLog )
+	{
+		FShaderResourceKey Key;
+		Key.Tex.Type = SCTT_Buffer;
+		Key.Tex.X = Vb->GetSize();
+		Key.Tex.Y = Vb->GetUsage();
+		Key.Tex.Z = Stride;
+		Key.Tex.Format = Format;
+		Key.bSRV = true;
+		SRVs.Add(SRV, Key);
+		CachedSRVs.Add(Key, FShaderResourceViewBinding(SRV, Vb, nullptr));
 		
-		if(bOK)
+		FShaderPlatformCache& PlatformCache = Caches.FindOrAdd(GMaxRHIShaderPlatform);
+		PlatformCache.Resources.Add(Key);
+	}
+}
+
+void FShaderCache::InternalRemoveSRV(FShaderResourceViewRHIParamRef SRV)
+{
+	if ( bUseShaderPredraw || bUseShaderDrawLog )
+	{
+		auto Key = SRVs.FindRef(SRV);
+		CachedSRVs.Remove(Key);
+		SRVs.Remove(SRV);
+	}
+}
+
+void FShaderCache::InternalRemoveTexture(FTextureRHIParamRef Texture)
+{
+	if ( bUseShaderPredraw || bUseShaderDrawLog )
+	{
+		auto Key = Textures.FindRef(Texture);
+		CachedTextures.Remove(Key);
+		Textures.Remove(Texture);
+	}
+}
+
+void FShaderCache::InternalSetBlendState(FBlendStateRHIParamRef State)
+{
+	if ( (bUseShaderPredraw || bUseShaderDrawLog) && !bIsPreDraw )
+	{
+		CurrentDrawKey.BlendState = BlendStates.FindChecked(State);
+	}
+}
+
+void FShaderCache::InternalSetRasterizerState(FRasterizerStateRHIParamRef State)
+{
+	if ( (bUseShaderPredraw || bUseShaderDrawLog) && !bIsPreDraw )
+	{
+		CurrentDrawKey.RasterizerState = RasterizerStates.FindChecked(State);
+	}
+}
+
+void FShaderCache::InternalSetDepthStencilState(FDepthStencilStateRHIParamRef State)
+{
+	if ( (bUseShaderPredraw || bUseShaderDrawLog) && !bIsPreDraw )
+	{
+		CurrentDrawKey.DepthStencilState = DepthStencilStates.FindChecked(State);
+	}
+}
+
+void FShaderCache::InternalSetRenderTargets( uint32 NumSimultaneousRenderTargets, const FRHIRenderTargetView* NewRenderTargetsRHI, const FRHIDepthRenderTargetView* NewDepthStencilTargetRHI )
+{
+	if ( bUseShaderDrawLog && !bIsPreDraw )
+	{
+		FShaderPlatformCache& PlatformCache = Caches.FindOrAdd(GMaxRHIShaderPlatform);
+		CurrentNumRenderTargets = NumSimultaneousRenderTargets;
+		bCurrentDepthStencilTarget = (NewDepthStencilTargetRHI != nullptr);
+		
+		FMemory::Memzero(CurrentRenderTargets, sizeof(FRHIRenderTargetView) * MaxSimultaneousRenderTargets);
+		FMemory::Memcpy(CurrentRenderTargets, NewRenderTargetsRHI, sizeof(FRHIRenderTargetView) * NumSimultaneousRenderTargets);
+		if( NewDepthStencilTargetRHI )
+		{
+			CurrentDepthStencilTarget = *NewDepthStencilTargetRHI;
+		}
+	
+		FMemory::Memset(CurrentDrawKey.RenderTargets, 255, sizeof(uint32) * MaxSimultaneousRenderTargets);
+		for( int32 RenderTargetIndex = NumSimultaneousRenderTargets - 1; RenderTargetIndex >= 0; --RenderTargetIndex )
+		{
+			FRHIRenderTargetView const& Target = NewRenderTargetsRHI[RenderTargetIndex];
+			if ( Target.Texture )
+			{
+				FShaderRenderTargetKey Key;
+				Key.Texture = Textures.FindChecked(Target.Texture);
+				check(Key.Texture.MipLevels == Target.Texture->GetNumMips());
+				Key.MipLevel = Key.Texture.MipLevels > Target.MipIndex ? Target.MipIndex : 0;
+				Key.ArrayIndex = Target.ArraySliceIndex;
+				CurrentDrawKey.RenderTargets[RenderTargetIndex] = PlatformCache.RenderTargets.Add(Key).AsInteger();
+			}
+			else
+			{
+				CurrentDrawKey.RenderTargets[RenderTargetIndex] = FShaderDrawKey::NullState;
+			}
+		}
+		
+		if ( NewDepthStencilTargetRHI && NewDepthStencilTargetRHI->Texture )
+		{
+			FShaderRenderTargetKey Key;
+			Key.Texture = Textures.FindChecked(NewDepthStencilTargetRHI->Texture);
+			CurrentDrawKey.DepthStencilTarget = PlatformCache.RenderTargets.Add(Key).AsInteger();
+		}
+		else
+		{
+			CurrentDrawKey.DepthStencilTarget = FShaderDrawKey::NullState;
+		}
+	}
+}
+
+void FShaderCache::InternalSetSamplerState(EShaderFrequency Frequency, uint32 Index, FSamplerStateRHIParamRef State)
+{
+	if ( bUseShaderDrawLog && !bIsPreDraw )
+	{
+		check(Index < GetFeatureLevelMaxTextureSamplers(GMaxRHIFeatureLevel));
+		if ( State )
+		{
+			FSamplerStateInitializerRHI Key = SamplerStates.FindChecked(State);
+			
+			FShaderPlatformCache& PlatformCache = Caches.FindOrAdd(GMaxRHIShaderPlatform);
+			CurrentDrawKey.SamplerStates[Frequency][Index] = PlatformCache.SamplerStates.FindId(Key).AsInteger();
+		}
+		else
+		{
+			CurrentDrawKey.SamplerStates[Frequency][Index] = FShaderDrawKey::NullState;
+		}
+	}
+}
+
+void FShaderCache::InternalSetTexture(EShaderFrequency Frequency, uint32 Index, FTextureRHIParamRef State)
+{
+	if ( bUseShaderDrawLog && !bIsPreDraw )
+	{
+		check(Index < GetFeatureLevelMaxTextureSamplers(GMaxRHIFeatureLevel));
+		if ( State )
 		{
-			FVertexDeclarationRHIRef VertexDeclaration = RHICreateVertexDeclaration(State.VertexDeclaration);
-			bOK &= IsValidRef(VertexDeclaration);
+			FShaderResourceKey Key;
 			
-			if(bOK)
-			{
-				FBoundShaderStateRHIRef BoundState = RHICreateBoundShaderState(
-					VertexDeclaration,
-					VertexShader,
-					HullShader,
-					DomainShader,
-					PixelShader,
-					GeometryShader
-				);
-				if(IsValidRef(BoundState))
+			FTextureRHIParamRef Tex = State;
+			if ( State->GetTextureReference() )
+			{
+				Tex = State->GetTextureReference()->GetReferencedTexture();
+			}
+			
+			Key.Tex = Textures.FindChecked(Tex);
+			Key.Format = Key.Tex.Format;
+			
+			FShaderPlatformCache& PlatformCache = Caches.FindOrAdd(GMaxRHIShaderPlatform);
+			CurrentDrawKey.Resources[Frequency][Index] = PlatformCache.Resources.FindId(Key).AsInteger();
+		}
+		else
+		{
+			CurrentDrawKey.Resources[Frequency][Index] = FShaderDrawKey::NullState;
+		}
+	}
+}
+
+void FShaderCache::InternalSetSRV(EShaderFrequency Frequency, uint32 Index, FShaderResourceViewRHIParamRef SRV)
+{
+	if ( bUseShaderDrawLog && !bIsPreDraw )
+	{
+		check(Index < GetFeatureLevelMaxTextureSamplers(GMaxRHIFeatureLevel));
+		if ( SRV )
+		{
+			FShaderResourceKey Key = SRVs.FindChecked(SRV);
+			
+			FShaderPlatformCache& PlatformCache = Caches.FindOrAdd(GMaxRHIShaderPlatform);
+			CurrentDrawKey.Resources[Frequency][Index] = PlatformCache.Resources.Add(Key).AsInteger();
+		}
+		else
+		{
+			CurrentDrawKey.Resources[Frequency][Index] = FShaderDrawKey::NullState;
+		}
+	}
+}
+
+void FShaderCache::InternalSetBoundShaderState(FBoundShaderStateRHIParamRef State)
+{
+	if ( (bUseShaderPredraw || bUseShaderDrawLog) && !bIsPreDraw )
+	{
+		FMemory::Memset(CurrentDrawKey.SamplerStates, 255, sizeof(CurrentDrawKey.SamplerStates));
+		FMemory::Memset(CurrentDrawKey.Resources, 255, sizeof(CurrentDrawKey.Resources));
+		CurrentShaderState = State;
+		if ( State )
+		{
+			BoundShaderState = ShaderStates.FindChecked(State);
+		}
+	}
+}
+
+void FShaderCache::InternalSetViewport(uint32 MinX, uint32 MinY, float MinZ, uint32 MaxX, uint32 MaxY, float MaxZ)
+{
+	if ( (bUseShaderPredraw || bUseShaderDrawLog) && !bIsPreDraw )
+	{
+		Viewport[0] = MinX;
+		Viewport[1] = MinY;
+		Viewport[2] = MaxX;
+		Viewport[3] = MaxY;
+		DepthRange[0] = MinZ;
+		DepthRange[1] = MaxZ;
+	}
+}
+
+void FShaderCache::InternalLogDraw(uint8 IndexType)
+{
+	if ( bUseShaderDrawLog && !bIsPreDraw )
+	{
+		FShaderPlatformCache& PlatformCache = Caches.FindOrAdd(GMaxRHIShaderPlatform);
+		CurrentDrawKey.IndexType = IndexType;
+		FSetElementId Id = PlatformCache.DrawStates.Add(CurrentDrawKey);
+		TSet<int32>& ShaderDrawSet = PlatformCache.ShaderDrawStates.FindOrAdd(BoundShaderState);
+		ShaderDrawSet.Add(Id.AsInteger());
+		
+		// No need to predraw this shader draw key - we've already done it
+		ShadersToDraw.FindRef(BoundShaderState).Remove(PlatformCache.DrawStates.FindId(CurrentDrawKey).AsInteger());
+	}
+}
+
+void FShaderCache::InternalPreDrawShaders(FRHICommandList& RHICmdList)
+{
+	if ( bUseShaderPredraw && ShadersToDraw.Num() > 0 )
+	{
+		bIsPreDraw = true;
+		
+		if ( !IsValidRef(IndexBufferUInt16) )
+		{
+			FRHIResourceCreateInfo Info;
+			uint32 Stride = sizeof(uint16);
+			uint32 Size = sizeof(uint16) * 3;
+			IndexBufferUInt16 = RHICreateIndexBuffer(Stride, Size, BUF_Static, Info);
+			void* Data = RHILockIndexBuffer(IndexBufferUInt16, 0, Size, RLM_WriteOnly);
+			if ( Data )
+			{
+				FMemory::Memzero(Data, Size);
+			}
+			RHIUnlockIndexBuffer(IndexBufferUInt16);
+		}
+		if ( !IsValidRef(IndexBufferUInt32) )
+		{
+			FRHIResourceCreateInfo Info;
+			uint32 Stride = sizeof(uint32);
+			uint32 Size = sizeof(uint32) * 3;
+			IndexBufferUInt32 = RHICreateIndexBuffer(Stride, Size, BUF_Static, Info);
+			void* Data = RHILockIndexBuffer(IndexBufferUInt32, 0, Size, RLM_WriteOnly);
+			if ( Data )
+			{
+				FMemory::Memzero(Data, Size);
+			}
+			RHIUnlockIndexBuffer(IndexBufferUInt32);
+		}
+		
+		RHICmdList.SetViewport(0, 0, FLT_MIN, 3, 3, FLT_MAX);
+		
+		int64 TimeForPredrawing = 0;
+		for ( auto It = ShadersToDraw.CreateIterator(); (PredrawBatchTime == -1 || TimeForPredrawing < PredrawBatchTime) && It; ++It )
+		{
+			uint32 Start = FPlatformTime::Cycles();
+			
+			auto Shader = *It;
+			TSet<int32>& ShaderDrawSet = Shader.Value;
+			PreDrawShader(RHICmdList, Shader.Key, ShaderDrawSet);
+			It.RemoveCurrent();
+			
+			uint32 End = FPlatformTime::Cycles();
+			TimeForPredrawing += FPlatformTime::ToMilliseconds(End - Start);
+		}
+		
+		// This is a bit dirty/naughty but it forces the draw commands to be flushed through on OS X
+		// which means we can delete the resources without crashing MTGL.
+		RHIFlushResources();
+		
+		RHICmdList.SetBoundShaderState(CurrentShaderState);
+		
+		FBlendStateRHIRef BlendState = RHICreateBlendState(CurrentDrawKey.BlendState);
+		FDepthStencilStateRHIRef DepthStencil = RHICreateDepthStencilState(CurrentDrawKey.DepthStencilState);
+		FRasterizerStateRHIRef Rasterizer = RHICreateRasterizerState(CurrentDrawKey.RasterizerState);
+		
+		RHICmdList.SetBlendState(BlendState);
+		RHICmdList.SetDepthStencilState(DepthStencil);
+		RHICmdList.SetRasterizerState(Rasterizer);
+
+		RHICmdList.SetViewport(Viewport[0], Viewport[1], DepthRange[0], Viewport[2], Viewport[3], DepthRange[1]);
+		
+		if ( ShadersToDraw.Num() == 0 )
+		{
+			PredrawRTs.Empty();
+			PredrawBindings.Empty();
+			PredrawVBs.Empty();
+		}
+		
+		bIsPreDraw = false;
+	}
+}
+
+void FShaderCache::PrebindShader(FShaderCacheKey const& Key)
+{
+	FShaderPlatformCache& PlatformCache = Caches.FindOrAdd(Key.Platform);
+	bool const bCanPreBind = (ShaderPlatformCanPrebindBoundShaderState(Key.Platform) || CurrentNumRenderTargets > 0);
+	if (bCanPreBind || bUseShaderPredraw)
+	{
+		TSet<FShaderCacheBoundState>& BoundStates = PlatformCache.ShaderStateMembership.FindOrAdd(Key);
+		for (FShaderCacheBoundState State : BoundStates)
+		{
+			FVertexShaderRHIRef VertexShader = State.VertexShader.bActive ? CachedVertexShaders.FindRef(State.VertexShader) : nullptr;
+			FPixelShaderRHIRef PixelShader = State.PixelShader.bActive ? CachedPixelShaders.FindRef(State.PixelShader) : nullptr;
+			FGeometryShaderRHIRef GeometryShader = State.GeometryShader.bActive ? CachedGeometryShaders.FindRef(State.GeometryShader) : nullptr;
+			FHullShaderRHIRef HullShader = State.HullShader.bActive ? CachedHullShaders.FindRef(State.HullShader) : nullptr;
+			FDomainShaderRHIRef DomainShader = State.DomainShader.bActive ? CachedDomainShaders.FindRef(State.DomainShader) : nullptr;
+
+			bool bOK = true;
+			bOK &= (State.VertexShader.bActive == IsValidRef(VertexShader));
+			bOK &= (State.PixelShader.bActive == IsValidRef(PixelShader));
+			bOK &= (State.GeometryShader.bActive == IsValidRef(GeometryShader));
+			bOK &= (State.HullShader.bActive == IsValidRef(HullShader));
+			bOK &= (State.DomainShader.bActive == IsValidRef(DomainShader));
+
+			if (bOK)
+			{
+				FVertexDeclarationRHIRef VertexDeclaration = RHICreateVertexDeclaration(State.VertexDeclaration);
+				bOK &= IsValidRef(VertexDeclaration);
+
+				if (bOK)
 				{
-					BoundShaderStates.Add(State, BoundState);
+					if (bCanPreBind)
+					{
+						FBoundShaderStateRHIRef BoundState = RHICreateBoundShaderState(VertexDeclaration,
+																					   VertexShader,
+																					   HullShader,
+																					   DomainShader,
+																					   PixelShader,
+																					   GeometryShader);
+						if (IsValidRef(BoundState))
+						{
+							BoundShaderStates.Add(State, BoundState);
+							if (bUseShaderPredraw)
+							{
+								ShadersToDraw.Add(State, PlatformCache.ShaderDrawStates.FindOrAdd(State));
+							}
+						}
+					}
+					else if (bUseShaderPredraw)
+					{
+						ShadersToDraw.Add(State, PlatformCache.ShaderDrawStates.FindOrAdd(State));
+					}
 				}
 			}
 		}
 	}
 }
 
-void FShaderCache::SubmitShader(FShaderCacheKey Key, TArray<uint8> const& Code)
+void FShaderCache::SubmitShader(FShaderCacheKey const& Key, TArray<uint8> const& Code)
 {
 	FShaderPlatformCache& PlatformCache = Caches.FindOrAdd(Key.Platform);
 	switch (Key.Frequency)
@@ -424,20 +910,411 @@ void FShaderCache::SubmitShader(FShaderCacheKey Key, TArray<uint8> const& Code)
 			}
 			break;
 		case SF_Compute:
-			if(!CachedComputeShaders.Find(Key))
+		{
+			bool const bCanPreBind = (ShaderPlatformCanPrebindBoundShaderState(Key.Platform) || CurrentNumRenderTargets > 0);
+			if (!CachedComputeShaders.Find(Key) && bCanPreBind)
 			{
 				FComputeShaderRHIRef Shader = RHICreateComputeShader(Code);
-				if(IsValidRef(Shader))
+				if (IsValidRef(Shader))
 				{
-					Shader->SetHash(Key.Hash);
+					// @todo WARNING: The RHI is responsible for hashing Compute shaders, unlike other stages because of how OpenGLDrv implements compute.
+					// Shader->SetHash(Key.Hash);
 					CachedComputeShaders.Add(Key, Shader);
 					PrebindShader(Key);
 					PlatformCache.Shaders.Add(Key);
 				}
 			}
 			break;
+		}
 		default:
 			check(false);
 			break;
 	}
 }
+
+FTextureRHIRef FShaderCache::CreateTexture(FShaderTextureKey const& TextureKey, bool const bCached)
+{
+	FTextureRHIRef Tex = bCached ? CachedTextures.FindRef(TextureKey) : nullptr;
+	
+	if ( !IsValidRef(Tex) )
+	{
+		FRHIResourceCreateInfo Info;
+		
+		switch(TextureKey.Type)
+		{
+			case SCTT_Texture2D:
+			{
+				Tex = RHICreateTexture2D(TextureKey.X, TextureKey.Y, TextureKey.Format, TextureKey.MipLevels, TextureKey.Samples, TextureKey.Flags, Info);
+				break;
+			}
+			case SCTT_Texture2DArray:
+			{
+				Tex = RHICreateTexture2DArray(TextureKey.X, TextureKey.Y, TextureKey.Z, TextureKey.Format, TextureKey.MipLevels, TextureKey.Flags, Info);
+				break;
+			}
+			case SCTT_Texture3D:
+			{
+				Tex = RHICreateTexture3D(TextureKey.X, TextureKey.Y, TextureKey.Z, TextureKey.Format, TextureKey.MipLevels, TextureKey.Flags, Info);
+				break;
+			}
+			case SCTT_TextureCube:
+			{
+				Tex = RHICreateTextureCube(TextureKey.X, TextureKey.Format, TextureKey.MipLevels, TextureKey.Flags, Info);
+				break;
+			}
+			case SCTT_TextureCubeArray:
+			{
+				Tex = RHICreateTextureCubeArray(TextureKey.X, TextureKey.Z, TextureKey.Format, TextureKey.MipLevels, TextureKey.Flags, Info);
+				break;
+			}
+			case SCTT_Buffer:
+			case SCTT_Texture1D:
+			case SCTT_Texture1DArray:
+			default:
+			{
+				check(false);
+				break;
+			}
+		}
+	}
+	return Tex;
+}
+
+FShaderCache::FShaderTextureBinding FShaderCache::CreateSRV(FShaderResourceKey const& ResourceKey)
+{
+	FShaderTextureBinding Binding = CachedSRVs.FindRef(ResourceKey);
+	if ( !IsValidRef(Binding.SRV) )
+	{
+		FShaderTextureKey const& TextureKey = ResourceKey.Tex;
+		switch(TextureKey.Type)
+		{
+			case SCTT_Buffer:
+			{
+				FRHIResourceCreateInfo Info;
+				Binding.VertexBuffer = RHICreateVertexBuffer(TextureKey.X, TextureKey.Y, Info);
+				Binding.SRV = RHICreateShaderResourceView(Binding.VertexBuffer, TextureKey.Z, TextureKey.Format);
+				break;
+			}
+			case SCTT_Texture2D:
+			{
+				Binding.Texture = CreateTexture(TextureKey, true);
+				Binding.SRV = RHICreateShaderResourceView(Binding.Texture->GetTexture2D(), ResourceKey.BaseMip, ResourceKey.MipLevels, ResourceKey.Format);
+				break;
+			}
+			default:
+			{
+				check(false);
+				break;
+			}
+		}
+	}
+	
+	return Binding;
+}
+
+FTextureRHIRef FShaderCache::CreateRenderTarget(FShaderRenderTargetKey const& TargetKey)
+{
+	FTextureRHIRef Texture;
+	if( TargetKey.Texture.Format != PF_Unknown )
+	{
+		Texture = PredrawRTs.FindRef(TargetKey);
+		if ( !IsValidRef(Texture) )
+		{
+			Texture = CreateTexture(TargetKey.Texture, false);
+			PredrawRTs.Add(TargetKey, Texture);
+		}
+	}
+	return Texture;
+}
+
+template <typename TShaderRHIRef>
+void FShaderCache::SetShaderSamplerTextures( FRHICommandList& RHICmdList, FShaderDrawKey const& DrawKey, EShaderFrequency Frequency, TShaderRHIRef Shader, bool bClear )
+{
+	FShaderPlatformCache& PlatformCache = Caches.FindOrAdd(GMaxRHIShaderPlatform);
+	
+	for ( uint32 i = 0; i < GetFeatureLevelMaxTextureSamplers(GMaxRHIFeatureLevel); i++ )
+	{
+		if ( DrawKey.SamplerStates[Frequency][i] != FShaderDrawKey::NullState )
+		{
+			FSamplerStateInitializerRHI SamplerInit = PlatformCache.SamplerStates[FSetElementId::FromInteger(DrawKey.SamplerStates[Frequency][i])];
+			FSamplerStateRHIRef State = RHICreateSamplerState(SamplerInit);
+			RHICmdList.SetShaderSampler(Shader, i, State);
+			
+			FShaderTextureBinding Bind;
+			if ( DrawKey.Resources[Frequency][i] != FShaderDrawKey::NullState )
+			{
+				FShaderResourceKey Resource = PlatformCache.Resources[FSetElementId::FromInteger(DrawKey.Resources[Frequency][i])];
+				if( Resource.bSRV == false )
+				{
+					if ( !bClear && Resource.Tex.Type != SCTT_Invalid )
+					{
+						Bind.Texture = CreateTexture(Resource.Tex, true);
+						RHICmdList.SetShaderTexture(Shader, i, Bind.Texture.GetReference());
+					}
+					else
+					{
+						RHICmdList.SetShaderTexture(Shader, i, nullptr);
+					}
+				}
+				else
+				{
+					if ( !bClear )
+					{
+						Bind = CreateSRV(Resource);
+						RHICmdList.SetShaderResourceViewParameter(Shader, i, Bind.SRV.GetReference());
+					}
+					else
+					{
+						RHICmdList.SetShaderResourceViewParameter(Shader, i, nullptr);
+					}
+				}
+			}
+			else
+			{
+				RHICmdList.SetShaderTexture(Shader, i, nullptr);
+			}
+			
+			if ( IsValidRef(Bind.Texture) || IsValidRef(Bind.SRV) )
+			{
+				PredrawBindings.Add(Bind);
+			}
+		}
+	}
+}
+
+void FShaderCache::PreDrawShader(FRHICommandList& RHICmdList, FShaderCacheBoundState const& Shader, TSet<int32> const& DrawStates)
+{
+	FBoundShaderStateRHIRef BoundShaderState = BoundShaderStates.FindRef( Shader );
+	{
+		uint32 VertexBufferSize = 0;
+		for( auto VertexDec : Shader.VertexDeclaration )
+		{
+			VertexBufferSize = VertexBufferSize > (uint32)(VertexDec.Stride + VertexDec.Offset) ? VertexBufferSize : (uint32)(VertexDec.Stride + VertexDec.Offset);
+		}
+		
+		FRHIResourceCreateInfo Info;
+		if ( VertexBufferSize > 0 && (( !IsValidRef(PredrawVB) || !IsValidRef(PredrawZVB) ) || PredrawVB->GetSize() < VertexBufferSize || PredrawZVB->GetSize() < VertexBufferSize) )
+		{
+			// Retain previous VBs for outstanding draws
+			PredrawVBs.Add(PredrawVB);
+			PredrawVBs.Add(PredrawZVB);
+			
+			PredrawVB = RHICreateVertexBuffer(VertexBufferSize, BUF_Static, Info);
+			{
+				void* Data = RHILockVertexBuffer(PredrawVB, 0, VertexBufferSize, RLM_WriteOnly);
+				if ( Data )
+				{
+					FMemory::Memzero(Data, VertexBufferSize);
+				}
+				RHIUnlockVertexBuffer(PredrawVB);
+			}
+			PredrawZVB = RHICreateVertexBuffer(VertexBufferSize, BUF_Static|BUF_ZeroStride, Info);
+			{
+				void* Data = RHILockVertexBuffer(PredrawZVB, 0, VertexBufferSize, RLM_WriteOnly);
+				if ( Data )
+				{
+					FMemory::Memzero(Data, VertexBufferSize);
+				}
+				RHIUnlockVertexBuffer(PredrawZVB);
+			}
+		}
+		
+		for( auto VertexDec : Shader.VertexDeclaration )
+		{
+			if ( VertexDec.Stride > 0 )
+			{
+				RHICmdList.SetStreamSource(VertexDec.StreamIndex, PredrawVB, VertexDec.Stride, VertexDec.Offset);
+			}
+			else
+			{
+				RHICmdList.SetStreamSource(VertexDec.StreamIndex, PredrawZVB, VertexDec.Stride, VertexDec.Offset);
+			}
+		}
+		
+		FShaderPlatformCache& PlatformCache = Caches.FindOrAdd(GMaxRHIShaderPlatform);
+		for ( auto DrawKeyIdx : DrawStates )
+		{
+			FShaderDrawKey DrawKey = PlatformCache.DrawStates[FSetElementId::FromInteger(DrawKeyIdx)];
+			
+			FBlendStateRHIRef BlendState = RHICreateBlendState(DrawKey.BlendState);
+			FDepthStencilStateRHIRef DepthStencil = RHICreateDepthStencilState(DrawKey.DepthStencilState);
+			FRasterizerStateRHIRef Rasterizer = RHICreateRasterizerState(DrawKey.RasterizerState);
+			
+			RHICmdList.SetBlendState(BlendState);
+			RHICmdList.SetDepthStencilState(DepthStencil);
+			RHICmdList.SetRasterizerState(Rasterizer);
+			
+			uint32 NewNumRenderTargets = 0;
+			FRHIRenderTargetView RenderTargets[MaxSimultaneousRenderTargets];
+			for ( uint32 i = 0; i < MaxSimultaneousRenderTargets; i++ )
+			{
+				if( DrawKey.RenderTargets[i] != FShaderDrawKey::NullState )
+				{
+					FShaderTextureBinding Bind;
+					FShaderRenderTargetKey RTKey = PlatformCache.RenderTargets[FSetElementId::FromInteger(DrawKey.RenderTargets[i])];
+					Bind.Texture = CreateRenderTarget(RTKey);
+					RenderTargets[i].MipIndex = Bind.Texture->GetNumMips() > RTKey.MipLevel ? RTKey.MipLevel : 0;
+					RenderTargets[i].ArraySliceIndex = RTKey.ArrayIndex;
+					PredrawBindings.Add(Bind);
+					RenderTargets[i].Texture = Bind.Texture;
+					NewNumRenderTargets++;
+				}
+				else
+				{
+					break;
+				}
+			}
+			
+			bool bDepthStencilTarget = (DrawKey.DepthStencilTarget != FShaderDrawKey::NullState);
+			FRHIDepthRenderTargetView DepthStencilTarget;
+			if ( bDepthStencilTarget )
+			{
+				FShaderTextureBinding Bind;
+				FShaderRenderTargetKey RTKey = PlatformCache.RenderTargets[FSetElementId::FromInteger(DrawKey.DepthStencilTarget)];
+				Bind.Texture = CreateRenderTarget(RTKey);
+				PredrawBindings.Add(Bind);
+				DepthStencilTarget.Texture = Bind.Texture;
+			}
+			
+			RHICmdList.SetRenderTargets(NewNumRenderTargets, RenderTargets, bDepthStencilTarget ? &DepthStencilTarget : nullptr, 0, nullptr);
+			
+			if ( !IsValidRef(BoundShaderState) )
+			{
+				FVertexShaderRHIRef VertexShader = Shader.VertexShader.bActive ? CachedVertexShaders.FindRef(Shader.VertexShader) : nullptr;
+				FPixelShaderRHIRef PixelShader = Shader.PixelShader.bActive ? CachedPixelShaders.FindRef(Shader.PixelShader) : nullptr;
+				FGeometryShaderRHIRef GeometryShader = Shader.GeometryShader.bActive ? CachedGeometryShaders.FindRef(Shader.GeometryShader) : nullptr;
+				FHullShaderRHIRef HullShader = Shader.HullShader.bActive ? CachedHullShaders.FindRef(Shader.HullShader) : nullptr;
+				FDomainShaderRHIRef DomainShader = Shader.DomainShader.bActive ? CachedDomainShaders.FindRef(Shader.DomainShader) : nullptr;
+				
+				bool bOK = true;
+				bOK &= (Shader.VertexShader.bActive == IsValidRef(VertexShader));
+				bOK &= (Shader.PixelShader.bActive == IsValidRef(PixelShader));
+				bOK &= (Shader.GeometryShader.bActive == IsValidRef(GeometryShader));
+				bOK &= (Shader.HullShader.bActive == IsValidRef(HullShader));
+				bOK &= (Shader.DomainShader.bActive == IsValidRef(DomainShader));
+				
+				if ( bOK )
+				{
+					FVertexDeclarationRHIRef VertexDeclaration = RHICreateVertexDeclaration(Shader.VertexDeclaration);
+					bOK &= IsValidRef(VertexDeclaration);
+					
+					if ( bOK )
+					{
+						BoundShaderState = RHICreateBoundShaderState( VertexDeclaration,
+																	 VertexShader,
+																	 HullShader,
+																	 DomainShader,
+																	 PixelShader,
+																	 GeometryShader );
+					}
+				}
+			}
+			
+			if ( IsValidRef( BoundShaderState ) )
+			{
+				RHICmdList.SetBoundShaderState(BoundShaderState);
+			}
+			else
+			{
+				break;
+			}
+			
+			if ( Shader.VertexShader.bActive )
+			{
+				SetShaderSamplerTextures(RHICmdList, DrawKey, SF_Vertex, CachedVertexShaders.FindRef(Shader.VertexShader).GetReference());
+			}
+			if ( Shader.PixelShader.bActive )
+			{
+				SetShaderSamplerTextures(RHICmdList, DrawKey, SF_Pixel, CachedPixelShaders.FindRef(Shader.PixelShader).GetReference());
+			}
+			if ( Shader.GeometryShader.bActive )
+			{
+				SetShaderSamplerTextures(RHICmdList, DrawKey, SF_Geometry, CachedGeometryShaders.FindRef(Shader.GeometryShader).GetReference());
+			}
+			if ( Shader.HullShader.bActive )
+			{
+				SetShaderSamplerTextures(RHICmdList, DrawKey, SF_Hull, CachedHullShaders.FindRef(Shader.HullShader).GetReference());
+			}
+			if ( Shader.DomainShader.bActive )
+			{
+				SetShaderSamplerTextures(RHICmdList, DrawKey, SF_Domain, CachedDomainShaders.FindRef(Shader.DomainShader).GetReference());
+			}
+			
+			switch ( DrawKey.IndexType )
+			{
+				case 0:
+				{
+					RHICmdList.DrawPrimitive(PT_TriangleList, 0, 1, 1);
+					break;
+				}
+				case 2:
+				{
+					RHICmdList.DrawIndexedPrimitive(IndexBufferUInt16, PT_TriangleList, 0, 0, 3, 0, 1, 1);
+					break;
+				}
+				case 4:
+				{
+					RHICmdList.DrawIndexedPrimitive(IndexBufferUInt32, PT_TriangleList, 0, 0, 3, 0, 1, 1);
+					break;
+				}
+				default:
+				{
+					break;
+				}
+			}
+		}
+		
+		if( IsValidRef( BoundShaderState ) && DrawStates.Num() )
+		{
+			if ( Shader.VertexShader.bActive )
+			{
+				SetShaderSamplerTextures(RHICmdList, CurrentDrawKey, SF_Vertex, CachedVertexShaders.FindRef(Shader.VertexShader).GetReference(), true);
+			}
+			if ( Shader.PixelShader.bActive )
+			{
+				SetShaderSamplerTextures(RHICmdList, CurrentDrawKey, SF_Pixel, CachedPixelShaders.FindRef(Shader.PixelShader).GetReference(), true);
+			}
+			if ( Shader.GeometryShader.bActive )
+			{
+				SetShaderSamplerTextures(RHICmdList, CurrentDrawKey, SF_Geometry, CachedGeometryShaders.FindRef(Shader.GeometryShader).GetReference(), true);
+			}
+			if ( Shader.HullShader.bActive )
+			{
+				SetShaderSamplerTextures(RHICmdList, CurrentDrawKey, SF_Hull, CachedHullShaders.FindRef(Shader.HullShader).GetReference(), true);
+			}
+			if ( Shader.DomainShader.bActive )
+			{
+				SetShaderSamplerTextures(RHICmdList, CurrentDrawKey, SF_Domain, CachedDomainShaders.FindRef(Shader.DomainShader).GetReference(), true);
+			}
+		}
+		
+		for( auto VertexDec : Shader.VertexDeclaration )
+		{
+			RHICmdList.SetStreamSource(VertexDec.StreamIndex, nullptr, 0, 0);
+		}
+	}
+}
+
+static FORCEINLINE uint32 CalculateSizeOfSamplerStateInitializer()
+{
+	static uint32 SizeOfSamplerStateInitializer = 0;
+	if (SizeOfSamplerStateInitializer == 0)
+	{
+		TArray<uint8> Data;
+		FMemoryWriter Writer(Data);
+		FSamplerStateInitializerRHI State;
+		Writer << State;
+		SizeOfSamplerStateInitializer = Data.Num();
+	}
+	return SizeOfSamplerStateInitializer;
+}
+
+bool FShaderCache::FSamplerStateInitializerRHIKeyFuncs::Matches(KeyInitType A,KeyInitType B)
+{
+	return FMemory::Memcmp(&A, &B, CalculateSizeOfSamplerStateInitializer()) == 0;
+}
+
+uint32 FShaderCache::FSamplerStateInitializerRHIKeyFuncs::GetKeyHash(KeyInitType Key)
+{
+	return FCrc::MemCrc_DEPRECATED(&Key, CalculateSizeOfSamplerStateInitializer());
+}
diff --git a/Engine/Source/Runtime/ShaderCore/Public/ShaderCache.h b/Engine/Source/Runtime/ShaderCore/Public/ShaderCache.h
--- a/Engine/Source/Runtime/ShaderCore/Public/ShaderCache.h
+++ b/Engine/Source/Runtime/ShaderCore/Public/ShaderCache.h
@@ -10,6 +10,143 @@
 #include "RHI.h"
 #include "BoundShaderStateCache.h"
 
+/** Custom serialization version for FShaderCache */
+struct SHADERCORE_API FShaderCacheCustomVersion
+{
+	static const FGuid Key;
+	static const FGuid GameKey;
+	enum Type {	Initial, PreDraw, Latest = PreDraw };
+};
+
+enum EShaderCacheTextureType
+{
+	SCTT_Invalid,
+	SCTT_Texture1D,
+	SCTT_Texture2D,
+	SCTT_Texture3D,
+	SCTT_TextureCube,
+	SCTT_Texture1DArray,
+	SCTT_Texture2DArray,
+	SCTT_TextureCubeArray,
+	SCTT_Buffer
+};
+
+struct SHADERCORE_API FShaderTextureKey
+{
+	uint32 X;
+	uint32 Y;
+	uint32 Z;
+	uint32 Flags; // ETextureCreateFlags
+	uint32 MipLevels;
+	uint32 Samples;
+	uint8 Format;
+	TEnumAsByte<EShaderCacheTextureType> Type;
+	
+	FShaderTextureKey()
+	: X(0)
+	, Y(0)
+	, Z(0)
+	, Flags(0)
+	, MipLevels(0)
+	, Samples(0)
+	, Format(PF_Unknown)
+	, Type(SCTT_Invalid)
+	{
+	}
+	
+	friend bool operator ==(const FShaderTextureKey& A,const FShaderTextureKey& B)
+	{
+		return A.X == B.X && A.Y == B.Y && A.Z == B.Z && A.Flags == B.Flags && A.MipLevels == B.MipLevels && A.Samples == B.Samples && A.Format == B.Format && A.Type == B.Type;
+	}
+	
+	friend uint32 GetTypeHash(const FShaderTextureKey &Key)
+	{
+		uint32 Hash = Key.X;
+		Hash ^= Key.Y;
+		Hash ^= Key.Z;
+		Hash ^= Key.Flags;
+		Hash ^= Key.MipLevels;
+		Hash ^= Key.Samples;
+		Hash ^= Key.Format;
+		Hash ^= Key.Type;
+		return Hash;
+	}
+		
+	friend FArchive& operator<<( FArchive& Ar, FShaderTextureKey& Info )
+	{
+		return Ar << Info.Format << Info.Type << Info.Samples << Info.MipLevels << Info.Flags << Info.X << Info.Y << Info.Z;
+	}
+};
+
+struct SHADERCORE_API FShaderResourceKey
+{
+	FShaderTextureKey Tex;
+	uint32 BaseMip;
+	uint32 MipLevels;
+	uint8 Format;
+	bool bSRV;
+	
+	FShaderResourceKey()
+	: BaseMip(0)
+	, MipLevels(0)
+	, Format(0)
+	, bSRV(false)
+	{
+	}
+	
+	friend bool operator ==(const FShaderResourceKey& A,const FShaderResourceKey& B)
+	{
+		return A.BaseMip == B.BaseMip && A.MipLevels == B.MipLevels && A.Format == B.Format && A.bSRV == B.bSRV && A.Tex == B.Tex;
+	}
+	
+	friend uint32 GetTypeHash(const FShaderResourceKey &Key)
+	{
+		uint32 Hash = GetTypeHash(Key.Tex);
+		Hash ^= Key.BaseMip;
+		Hash ^= Key.MipLevels;
+		Hash ^= Key.Format;
+		Hash ^= (uint32)Key.bSRV;
+		return Hash;
+	}
+	
+	friend FArchive& operator<<( FArchive& Ar, FShaderResourceKey& Info )
+	{
+		return Ar << Info.Tex << Info.BaseMip << Info.MipLevels << Info.Format << Info.bSRV;
+	}
+};
+
+struct SHADERCORE_API FShaderRenderTargetKey
+{
+	FShaderTextureKey Texture;
+	uint32 MipLevel;
+	uint32 ArrayIndex;
+	
+	FShaderRenderTargetKey()
+	: MipLevel(0)
+	, ArrayIndex(0)
+	{
+		
+	}
+	
+	friend bool operator ==(const FShaderRenderTargetKey& A,const FShaderRenderTargetKey& B)
+	{
+		return A.MipLevel == B.MipLevel && A.ArrayIndex == B.ArrayIndex && A.Texture == B.Texture;
+	}
+	
+	friend uint32 GetTypeHash(const FShaderRenderTargetKey &Key)
+	{
+		uint32 Hash = GetTypeHash(Key.Texture);
+		Hash ^= Key.MipLevel;
+		Hash ^= Key.ArrayIndex;
+		return Hash;
+	}
+	
+	friend FArchive& operator<<( FArchive& Ar, FShaderRenderTargetKey& Info )
+	{
+		return Ar << Info.Texture << Info.MipLevel << Info.ArrayIndex;
+	}
+};
+
 class SHADERCORE_API FShaderCache
 {
 	struct SHADERCORE_API FShaderCacheKey
@@ -43,42 +180,342 @@ class SHADERCORE_API FShaderCache
 			return Ar << Info.Hash << Info.bActive;
 		}
 	};
+	
+	struct SHADERCORE_API FShaderDrawKey
+	{
+		enum
+		{
+			MaxTextureSamplers = 16,
+			NullState = ~0u
+		};
+		
+		FShaderDrawKey()
+		: DepthStencilTarget(NullState)
+		, IndexType(0)
+		{
+			FMemory::Memzero(&BlendState, sizeof(BlendState));
+			FMemory::Memzero(&RasterizerState, sizeof(RasterizerState));
+			FMemory::Memzero(&DepthStencilState, sizeof(DepthStencilState));
+			FMemory::Memset(RenderTargets, 255, sizeof(RenderTargets));
+			FMemory::Memset(SamplerStates, 255, sizeof(SamplerStates));
+			FMemory::Memset(Resources, 255, sizeof(Resources));
+			check(GetFeatureLevelMaxTextureSamplers(GMaxRHIFeatureLevel) <= MaxTextureSamplers);
+		}
+		
+		struct FShaderRasterizerState
+		{
+			FShaderRasterizerState() { FMemory::Memzero(*this); }
+			FShaderRasterizerState(FRasterizerStateInitializerRHI const& Other) { operator=(Other); }
+			
+			float DepthBias;
+			float SlopeScaleDepthBias;
+			TEnumAsByte<ERasterizerFillMode> FillMode;
+			TEnumAsByte<ERasterizerCullMode> CullMode;
+			bool bAllowMSAA;
+			bool bEnableLineAA;
+			
+			FShaderRasterizerState& operator=(FRasterizerStateInitializerRHI const& Other)
+			{
+				DepthBias = Other.DepthBias;
+				SlopeScaleDepthBias = Other.SlopeScaleDepthBias;
+				FillMode = Other.FillMode;
+				CullMode = Other.CullMode;
+				bAllowMSAA = Other.bAllowMSAA;
+				bEnableLineAA = Other.bEnableLineAA;
+				return *this;
+			}
+			
+			operator FRasterizerStateInitializerRHI () const
+			{
+				FRasterizerStateInitializerRHI Initializer = {FillMode, CullMode, DepthBias, SlopeScaleDepthBias, bAllowMSAA, bEnableLineAA};
+				return Initializer;
+			}
+			
+			friend FArchive& operator<<(FArchive& Ar,FShaderRasterizerState& RasterizerStateInitializer)
+			{
+				Ar << RasterizerStateInitializer.DepthBias;
+				Ar << RasterizerStateInitializer.SlopeScaleDepthBias;
+				Ar << RasterizerStateInitializer.FillMode;
+				Ar << RasterizerStateInitializer.CullMode;
+				Ar << RasterizerStateInitializer.bAllowMSAA;
+				Ar << RasterizerStateInitializer.bEnableLineAA;
+				return Ar;
+			}
+		};
+	
+		FBlendStateInitializerRHI BlendState;
+		FShaderRasterizerState RasterizerState;
+		FDepthStencilStateInitializerRHI DepthStencilState;
+		uint32 RenderTargets[MaxSimultaneousRenderTargets];
+		uint32 SamplerStates[SF_NumFrequencies][MaxTextureSamplers];
+		uint32 Resources[SF_NumFrequencies][MaxTextureSamplers];
+		uint32 DepthStencilTarget;
+		uint8 IndexType;
+		
+		friend bool operator ==(const FShaderDrawKey& A,const FShaderDrawKey& B)
+		{
+			bool Compare = A.IndexType == B.IndexType &&
+							A.DepthStencilTarget == B.DepthStencilTarget &&
+							FMemory::Memcmp(&A.BlendState, &B.BlendState, sizeof(FBlendStateInitializerRHI)) == 0 &&
+							FMemory::Memcmp(&A.RasterizerState, &B.RasterizerState, sizeof(FShaderRasterizerState)) == 0 &&
+							FMemory::Memcmp(&A.DepthStencilState, &B.DepthStencilState, sizeof(FDepthStencilStateInitializerRHI)) == 0 &&
+							FMemory::Memcmp(&A.RenderTargets, &B.RenderTargets, sizeof(A.RenderTargets)) == 0 &&
+							FMemory::Memcmp(&A.SamplerStates, &B.SamplerStates, sizeof(A.SamplerStates)) == 0 &&
+							FMemory::Memcmp(&A.Resources, &B.Resources, sizeof(A.Resources)) == 0;
+			return Compare;
+		}
+		
+		friend uint32 GetTypeHash(const FShaderDrawKey &Key)
+		{
+			uint32 CRC = FCrc::MemCrc32(&Key.BlendState, sizeof(FBlendStateInitializerRHI));
+			CRC ^= FCrc::MemCrc32(&Key.RasterizerState, sizeof(FShaderRasterizerState));
+			CRC ^= FCrc::MemCrc32(&Key.DepthStencilState, sizeof(FDepthStencilStateInitializerRHI));
+			CRC ^= FCrc::MemCrc32(&Key.RenderTargets, sizeof(Key.RenderTargets));
+			CRC ^= FCrc::MemCrc32(&Key.SamplerStates, sizeof(Key.SamplerStates));
+			CRC ^= FCrc::MemCrc32(&Key.Resources, sizeof(Key.Resources));
+			CRC ^= GetTypeHash(Key.DepthStencilTarget);
+			CRC ^= GetTypeHash(Key.IndexType);
+			return CRC;
+		}
+		
+		friend FArchive& operator<<( FArchive& Ar, FShaderDrawKey& Info )
+		{
+			for ( uint32 i = 0; i < SF_NumFrequencies; i++ )
+			{
+				for ( uint32 j = 0; j < MaxTextureSamplers; j++ )
+				{
+					Ar << Info.SamplerStates[i][j];
+					Ar << Info.Resources[i][j];
+				}
+			}
+			for ( uint32 i = 0; i < MaxSimultaneousRenderTargets; i++ )
+			{
+				Ar << Info.RenderTargets[i];
+			}
+			return Ar << Info.BlendState << Info.RasterizerState << Info.DepthStencilState << Info.DepthStencilTarget << Info.IndexType;
+		}
+	};
 		
 public:
 	FShaderCache();
 	~FShaderCache();
 
+	// Called by the game to set the game specific shader cache version, only caches of this version will be loaded.
+	// Must be called before RHI initialisation, as InitShaderCache will load any existing cache.
+	static void SetGameVersion(int32 InGameVersion);
+
+	// Called by the RHI implementation, not to be called by the game.
 	static void InitShaderCache();
-	static FShaderCache* GetShaderCache();
 	static void ShutdownShaderCache();
+	
+	static FORCEINLINE FShaderCache* GetShaderCache()
+	{
+		return bUseShaderCaching ? Cache : nullptr;
+	}
 
 	FVertexShaderRHIRef GetVertexShader(EShaderPlatform Platform, FSHAHash Hash, TArray<uint8> const& Code);
 	FPixelShaderRHIRef GetPixelShader(EShaderPlatform Platform, FSHAHash Hash, TArray<uint8> const& Code);
 	FGeometryShaderRHIRef GetGeometryShader(EShaderPlatform Platform, FSHAHash Hash, TArray<uint8> const& Code);
 	FHullShaderRHIRef GetHullShader(EShaderPlatform Platform, FSHAHash Hash, TArray<uint8> const& Code);
 	FDomainShaderRHIRef GetDomainShader(EShaderPlatform Platform, FSHAHash Hash, TArray<uint8> const& Code);
-	FComputeShaderRHIRef GetComputeShader(EShaderPlatform Platform, FSHAHash Hash, TArray<uint8> const& Code);
+	FComputeShaderRHIRef GetComputeShader(EShaderPlatform Platform, TArray<uint8> const& Code);
 	FVertexDeclarationRHIParamRef GetVertexDeclaration(FVertexDeclarationElementList& VertexElements);
 	
-	void LogShader(EShaderPlatform Platform, EShaderFrequency Frequency, FSHAHash Hash, TArray<uint8> const& Code);
+	static FORCEINLINE void LogShader(EShaderPlatform Platform, EShaderFrequency Frequency, FSHAHash Hash, TArray<uint8> const& Code)
+	{
+		FShaderCache* Cache = GetShaderCache();
+		if ( Cache )
+		{
+			Cache->InternalLogShader(Platform, Frequency, Hash, Code);
+		}
+	}
 	
-	void LogVertexDeclaration(const FVertexDeclarationElementList& VertexElements, FVertexDeclarationRHIParamRef VertexDeclaration);
+	static FORCEINLINE void LogVertexDeclaration(const FVertexDeclarationElementList& VertexElements, FVertexDeclarationRHIParamRef VertexDeclaration)
+	{
+		FShaderCache* Cache = GetShaderCache();
+		if ( Cache )
+		{
+			Cache->InternalLogVertexDeclaration(VertexElements, VertexDeclaration);
+		}
+	}
 	
-	void LogBoundShaderState(EShaderPlatform Platform, FVertexDeclarationRHIParamRef VertexDeclaration,
-								FVertexShaderRHIParamRef VertexShader,
-								FPixelShaderRHIParamRef PixelShader,
-								FHullShaderRHIParamRef HullShader,
-								FDomainShaderRHIParamRef DomainShader,
-								FGeometryShaderRHIParamRef GeometryShader);
-		
-	friend FArchive& operator<<( FArchive& Ar, FShaderCache& Info )
+	static FORCEINLINE void LogBoundShaderState(EShaderPlatform Platform, FVertexDeclarationRHIParamRef VertexDeclaration,
+												FVertexShaderRHIParamRef VertexShader,
+												FPixelShaderRHIParamRef PixelShader,
+												FHullShaderRHIParamRef HullShader,
+												FDomainShaderRHIParamRef DomainShader,
+												FGeometryShaderRHIParamRef GeometryShader,
+												FBoundShaderStateRHIParamRef BoundState)
 	{
-		return Ar << Info.Caches;
+		FShaderCache* Cache = GetShaderCache();
+		if ( Cache )
+		{
+			Cache->InternalLogBoundShaderState(Platform, VertexDeclaration, VertexShader, PixelShader, HullShader, DomainShader, GeometryShader, BoundState);
+		}
 	}
 	
-private:
-	void PrebindShader(FShaderCacheKey Key);
-	void SubmitShader(FShaderCacheKey Key, TArray<uint8> const& Code);
+	static FORCEINLINE void LogBlendState(FBlendStateInitializerRHI const& Init, FBlendStateRHIParamRef State)
+	{
+		FShaderCache* Cache = GetShaderCache();
+		if ( Cache )
+		{
+			Cache->InternalLogBlendState(Init, State);
+		}
+	}
+	static FORCEINLINE void LogRasterizerState(FRasterizerStateInitializerRHI const& Init, FRasterizerStateRHIParamRef State)
+	{
+		FShaderCache* Cache = GetShaderCache();
+		if ( Cache )
+		{
+			Cache->InternalLogRasterizerState(Init, State);
+		}
+	}
+	static FORCEINLINE void LogDepthStencilState(FDepthStencilStateInitializerRHI const& Init, FDepthStencilStateRHIParamRef State)
+	{
+		FShaderCache* Cache = GetShaderCache();
+		if ( Cache )
+		{
+			Cache->InternalLogDepthStencilState(Init, State);
+		}
+	}
+	static FORCEINLINE void LogSamplerState(FSamplerStateInitializerRHI const& Init, FSamplerStateRHIParamRef State)
+	{
+		FShaderCache* Cache = GetShaderCache();
+		if ( Cache )
+		{
+			Cache->InternalLogSamplerState(Init, State);
+		}
+	}
+	static FORCEINLINE void LogTexture(FShaderTextureKey const& Init, FTextureRHIParamRef State)
+	{
+		FShaderCache* Cache = GetShaderCache();
+		if ( Cache )
+		{
+			Cache->InternalLogTexture(Init, State);
+		}
+	}
+	static FORCEINLINE void LogSRV(FShaderResourceViewRHIParamRef SRV, FTextureRHIParamRef Texture, uint8 StartMip, uint8 NumMips, uint8 Format)
+	{
+		FShaderCache* Cache = GetShaderCache();
+		if ( Cache )
+		{
+			Cache->InternalLogSRV(SRV, Texture, StartMip, NumMips, Format);
+		}
+	}
+	static FORCEINLINE void LogSRV(FShaderResourceViewRHIParamRef SRV, FVertexBufferRHIParamRef Vb, uint32 Stride, uint8 Format)
+	{
+		FShaderCache* Cache = GetShaderCache();
+		if ( Cache )
+		{
+			Cache->InternalLogSRV(SRV, Vb, Stride, Format);
+		}
+	}
+	static FORCEINLINE void RemoveSRV(FShaderResourceViewRHIParamRef SRV)
+	{
+		FShaderCache* Cache = GetShaderCache();
+		if ( Cache )
+		{
+			Cache->InternalRemoveSRV(SRV);
+		}
+	}
+	static FORCEINLINE void RemoveTexture(FTextureRHIParamRef Texture)
+	{
+		FShaderCache* Cache = GetShaderCache();
+		if ( Cache )
+		{
+			Cache->InternalRemoveTexture(Texture);
+		}
+	}
+	static FORCEINLINE void SetBlendState(FBlendStateRHIParamRef State)
+	{
+		FShaderCache* Cache = GetShaderCache();
+		if ( Cache )
+		{
+			Cache->InternalSetBlendState(State);
+		}
+	}
+	static FORCEINLINE void SetRasterizerState(FRasterizerStateRHIParamRef State)
+	{
+		FShaderCache* Cache = GetShaderCache();
+		if ( Cache )
+		{
+			Cache->InternalSetRasterizerState(State);
+		}
+	}
+	static FORCEINLINE void SetDepthStencilState(FDepthStencilStateRHIParamRef State)
+	{
+		FShaderCache* Cache = GetShaderCache();
+		if ( Cache )
+		{
+			Cache->InternalSetDepthStencilState(State);
+		}
+	}
+	static FORCEINLINE void SetRenderTargets( uint32 NumSimultaneousRenderTargets, const FRHIRenderTargetView* NewRenderTargetsRHI, const FRHIDepthRenderTargetView* NewDepthStencilTargetRHI )
+	{
+		FShaderCache* Cache = GetShaderCache();
+		if ( Cache )
+		{
+			Cache->InternalSetRenderTargets(NumSimultaneousRenderTargets, NewRenderTargetsRHI, NewDepthStencilTargetRHI);
+		}
+	}
+	static FORCEINLINE void SetSamplerState(EShaderFrequency Frequency, uint32 Index, FSamplerStateRHIParamRef State)
+	{
+		FShaderCache* Cache = GetShaderCache();
+		if ( Cache )
+		{
+			Cache->InternalSetSamplerState(Frequency, Index, State);
+		}
+	}
+	static FORCEINLINE void SetTexture(EShaderFrequency Frequency, uint32 Index, FTextureRHIParamRef State)
+	{
+		FShaderCache* Cache = GetShaderCache();
+		if ( Cache )
+		{
+			Cache->InternalSetTexture(Frequency, Index, State);
+		}
+	}
+	static FORCEINLINE void SetSRV(EShaderFrequency Frequency, uint32 Index, FShaderResourceViewRHIParamRef SRV)
+	{
+		FShaderCache* Cache = GetShaderCache();
+		if ( Cache )
+		{
+			Cache->InternalSetSRV(Frequency, Index, SRV);
+		}
+	}
+	static FORCEINLINE void SetBoundShaderState(FBoundShaderStateRHIParamRef State)
+	{
+		FShaderCache* Cache = GetShaderCache();
+		if ( Cache )
+		{
+			Cache->InternalSetBoundShaderState(State);
+		}
+	}
+	static FORCEINLINE void SetViewport(uint32 MinX, uint32 MinY, float MinZ, uint32 MaxX, uint32 MaxY, float MaxZ)
+	{
+		FShaderCache* Cache = GetShaderCache();
+		if ( Cache )
+		{
+			Cache->InternalSetViewport(MinX, MinY, MinZ, MaxX, MaxY, MaxZ);
+		}
+	}
+	static FORCEINLINE void LogDraw(uint8 IndexType)
+	{
+		FShaderCache* Cache = GetShaderCache();
+		if ( Cache )
+		{
+			Cache->InternalLogDraw(IndexType);
+		}
+	}
+	static FORCEINLINE void PreDrawShaders(FRHICommandList& RHICmdList)
+	{
+		FShaderCache* Cache = GetShaderCache();
+		if ( Cache )
+		{
+			Cache->InternalPreDrawShaders(RHICmdList);
+		}
+	}
+		
+	friend FArchive& operator<<( FArchive& Ar, FShaderCache& Info );
 	
 private:
 	struct FShaderCacheBoundState
@@ -123,25 +560,152 @@ private:
 			return Ar << Info.VertexDeclaration << Info.VertexShader << Info.PixelShader << Info.GeometryShader << Info.HullShader << Info.DomainShader;
 		}
 	};
+	
+	struct FSamplerStateInitializerRHIKeyFuncs : BaseKeyFuncs<FSamplerStateInitializerRHI,FSamplerStateInitializerRHI,false>
+	{
+		typedef TCallTraits<FSamplerStateInitializerRHI>::ParamType KeyInitType;
+		typedef TCallTraits<FSamplerStateInitializerRHI>::ParamType ElementInitType;
+		
+		/**
+		 * @return The key used to index the given element.
+		 */
+		static FORCEINLINE KeyInitType GetSetKey(ElementInitType Element)
+		{
+			return Element;
+		}
+		
+		/**
+		 * @return True if the keys match.
+		 */
+		static bool Matches(KeyInitType A,KeyInitType B);
+		
+		/** Calculates a hash index for a key. */
+		static uint32 GetKeyHash(KeyInitType Key);
+	};
 
 	struct FShaderPlatformCache
 	{
 		TSet<FShaderCacheKey> Shaders;
 		TSet<FShaderCacheBoundState> BoundShaderStates;
+		TSet<FShaderDrawKey> DrawStates;
+		TSet<FShaderRenderTargetKey> RenderTargets;
+		TSet<FShaderResourceKey> Resources;
+		TSet<FSamplerStateInitializerRHI, FSamplerStateInitializerRHIKeyFuncs> SamplerStates;
 		TMap<FShaderCacheKey, TSet<FShaderCacheBoundState>> ShaderStateMembership;
+		TMap<FShaderCacheBoundState, TSet<int32>> ShaderDrawStates;
 		
 		friend FArchive& operator<<( FArchive& Ar, FShaderPlatformCache& Info )
 		{
-			return Ar << Info.Shaders << Info.BoundShaderStates << Info.ShaderStateMembership;
+			return Ar << Info.Shaders << Info.BoundShaderStates << Info.DrawStates << Info.RenderTargets << Info.Resources << Info.SamplerStates << Info.ShaderStateMembership << Info.ShaderDrawStates;
 		}
 	};
 
 	typedef TMap<uint32, FShaderPlatformCache> FShaderCaches;
 
+	struct FShaderResourceViewBinding
+	{
+		FShaderResourceViewBinding()
+		: SRV(nullptr)
+		, VertexBuffer(nullptr)
+		, Texture(nullptr)
+		{}
+		
+		FShaderResourceViewBinding(FShaderResourceViewRHIParamRef InSrv, FVertexBufferRHIParamRef InVb, FTextureRHIParamRef InTexture)
+		: SRV(nullptr)
+		, VertexBuffer(nullptr)
+		, Texture(nullptr)
+		{}
+		
+		FShaderResourceViewRHIParamRef SRV;
+		FVertexBufferRHIParamRef VertexBuffer;
+		FTextureRHIParamRef Texture;
+	};
+
+	struct FShaderTextureBinding
+	{
+		FShaderTextureBinding() {}
+		FShaderTextureBinding(FShaderResourceViewBinding const& Other)
+		{
+			operator=(Other);
+		}
+		
+		FShaderTextureBinding& operator=(FShaderResourceViewBinding const& Other)
+		{
+			SRV = Other.SRV;
+			VertexBuffer = Other.VertexBuffer;
+			Texture = Other.Texture;
+			return *this;
+		}
+		
+		friend bool operator ==(const FShaderTextureBinding& A,const FShaderTextureBinding& B)
+		{
+			return A.SRV == B.SRV && A.VertexBuffer == B.VertexBuffer && A.Texture == B.Texture;
+		}
+		
+		friend uint32 GetTypeHash(const FShaderTextureBinding &Key)
+		{
+			uint32 Hash = GetTypeHash(Key.SRV);
+			Hash ^= GetTypeHash(Key.VertexBuffer);
+			Hash ^= GetTypeHash(Key.Texture);
+			return Hash;
+		}
+		
+		FShaderResourceViewRHIRef SRV;
+		FVertexBufferRHIRef VertexBuffer;
+		FTextureRHIRef Texture;
+	};
+
+private:
+	void InternalLogShader(EShaderPlatform Platform, EShaderFrequency Frequency, FSHAHash Hash, TArray<uint8> const& Code);
+	
+	void InternalLogVertexDeclaration(const FVertexDeclarationElementList& VertexElements, FVertexDeclarationRHIParamRef VertexDeclaration);
+	
+	void InternalLogBoundShaderState(EShaderPlatform Platform, FVertexDeclarationRHIParamRef VertexDeclaration,
+								FVertexShaderRHIParamRef VertexShader,
+								FPixelShaderRHIParamRef PixelShader,
+								FHullShaderRHIParamRef HullShader,
+								FDomainShaderRHIParamRef DomainShader,
+								FGeometryShaderRHIParamRef GeometryShader,
+								FBoundShaderStateRHIParamRef BoundState);
+	
+	void InternalLogBlendState(FBlendStateInitializerRHI const& Init, FBlendStateRHIParamRef State);
+	void InternalLogRasterizerState(FRasterizerStateInitializerRHI const& Init, FRasterizerStateRHIParamRef State);
+	void InternalLogDepthStencilState(FDepthStencilStateInitializerRHI const& Init, FDepthStencilStateRHIParamRef State);
+	void InternalLogSamplerState(FSamplerStateInitializerRHI const& Init, FSamplerStateRHIParamRef State);
+	void InternalLogTexture(FShaderTextureKey const& Init, FTextureRHIParamRef State);
+	void InternalLogSRV(FShaderResourceViewRHIParamRef SRV, FTextureRHIParamRef Texture, uint8 StartMip, uint8 NumMips, uint8 Format);
+	void InternalLogSRV(FShaderResourceViewRHIParamRef SRV, FVertexBufferRHIParamRef Vb, uint32 Stride, uint8 Format);
+	
+	void InternalRemoveSRV(FShaderResourceViewRHIParamRef SRV);
+	void InternalRemoveTexture(FTextureRHIParamRef Texture);
+	
+	void InternalSetBlendState(FBlendStateRHIParamRef State);
+	void InternalSetRasterizerState(FRasterizerStateRHIParamRef State);
+	void InternalSetDepthStencilState(FDepthStencilStateRHIParamRef State);
+	void InternalSetRenderTargets( uint32 NumSimultaneousRenderTargets, const FRHIRenderTargetView* NewRenderTargetsRHI, const FRHIDepthRenderTargetView* NewDepthStencilTargetRHI );
+	void InternalSetSamplerState(EShaderFrequency Frequency, uint32 Index, FSamplerStateRHIParamRef State);
+	void InternalSetTexture(EShaderFrequency Frequency, uint32 Index, FTextureRHIParamRef State);
+	void InternalSetSRV(EShaderFrequency Frequency, uint32 Index, FShaderResourceViewRHIParamRef SRV);
+	void InternalSetBoundShaderState(FBoundShaderStateRHIParamRef State);
+	void InternalSetViewport(uint32 MinX, uint32 MinY, float MinZ, uint32 MaxX, uint32 MaxY, float MaxZ);
+	
+	void InternalLogDraw(uint8 IndexType);
+	void InternalPreDrawShaders(FRHICommandList& RHICmdList);
+	
+	void PrebindShader(FShaderCacheKey const& Key);
+	void SubmitShader(FShaderCacheKey const& Key, TArray<uint8> const& Code);
+	void PreDrawShader(FRHICommandList& RHICmdList, FShaderCacheBoundState const& Shader, TSet<int32> const& DrawStates);
+	template <typename TShaderRHIRef>
+	void SetShaderSamplerTextures( FRHICommandList& RHICmdList, FShaderDrawKey const& DrawKey, EShaderFrequency Frequency, TShaderRHIRef Shader, bool bClear = false );
+	FTextureRHIRef CreateTexture(FShaderTextureKey const& TextureKey, bool const bCached);
+	FShaderTextureBinding CreateSRV(FShaderResourceKey const& ResourceKey);
+	FTextureRHIRef CreateRenderTarget(FShaderRenderTargetKey const& TargetKey);
+	
+private:
 	// Serialised
 	FShaderCaches Caches;
 	
-	// Transient
+	// Transient non-invasive tracking of RHI resources for shader logging
 	TMap<FShaderCacheKey, FVertexShaderRHIRef> CachedVertexShaders;
 	TMap<FShaderCacheKey, FPixelShaderRHIRef> CachedPixelShaders;
 	TMap<FShaderCacheKey, FGeometryShaderRHIRef> CachedGeometryShaders;
@@ -151,5 +715,56 @@ private:
 	TMap<FVertexDeclarationRHIParamRef, FVertexDeclarationElementList> VertexDeclarations;
 	TMap<FShaderCacheBoundState, FBoundShaderStateRHIRef> BoundShaderStates;
 	
+	// Transient non-invasive tracking of RHI resources for shader predrawing
+	TMap<FBoundShaderStateRHIParamRef, FShaderCacheBoundState> ShaderStates;
+	TMap<FBlendStateRHIParamRef, FBlendStateInitializerRHI> BlendStates;
+	TMap<FRasterizerStateRHIParamRef, FRasterizerStateInitializerRHI> RasterizerStates;
+	TMap<FDepthStencilStateRHIParamRef, FDepthStencilStateInitializerRHI> DepthStencilStates;
+	TMap<FSamplerStateRHIParamRef, FSamplerStateInitializerRHI> SamplerStates;
+	TMap<FTextureRHIParamRef, FShaderTextureKey> Textures;
+	TMap<FShaderResourceViewRHIParamRef, FShaderResourceKey> SRVs;
+	
+	// Shaders we need to predraw
+	TMap<FShaderCacheBoundState, TSet<int32>> ShadersToDraw;
+	
+	// Caches to track application & predraw created textures/SRVs so that we minimise temporary resource creation
+	TMap<FShaderTextureKey, FTextureRHIParamRef> CachedTextures;
+	TMap<FShaderResourceKey, FShaderResourceViewBinding> CachedSRVs;
+	
+	// Temporary shader resources for pre-draw
+	// Cleared after each round of pre-drawing is complete
+	TSet<FShaderTextureBinding> PredrawBindings;
+	TMap<FShaderRenderTargetKey, FTextureRHIParamRef> PredrawRTs;
+	TSet<FVertexBufferRHIRef> PredrawVBs;
+	
+	// Permanent shader pre-draw resources
+	FIndexBufferRHIRef IndexBufferUInt16;
+	FIndexBufferRHIRef IndexBufferUInt32;
+	
+	// Growable pre-draw resources
+	FVertexBufferRHIRef PredrawVB; // Standard VB
+	FVertexBufferRHIRef PredrawZVB; // Zero-stride VB
+	
+	// Current render state for draw state tracking
+	bool bCurrentDepthStencilTarget;
+	uint32 CurrentNumRenderTargets;
+	FRHIDepthRenderTargetView CurrentDepthStencilTarget;
+	FRHIRenderTargetView CurrentRenderTargets[MaxSimultaneousRenderTargets];
+	FShaderDrawKey CurrentDrawKey;
+	FBoundShaderStateRHIParamRef CurrentShaderState;
+	FShaderCacheBoundState BoundShaderState;
+	uint32 Viewport[4];
+	float DepthRange[2];
+	bool bIsPreDraw;
+	
 	static FShaderCache* Cache;
+	static int32 GameVersion;
+	static int32 bUseShaderCaching;
+	static int32 bUseShaderPredraw;
+	static int32 bUseShaderDrawLog;
+	static int32 PredrawBatchTime;
+	static FAutoConsoleVariableRef CVarUseShaderCaching;
+	static FAutoConsoleVariableRef CVarUseShaderPredraw;
+	static FAutoConsoleVariableRef CVarUseShaderDrawLog;
+	static FAutoConsoleVariableRef CVarPredrawBatchTime;
 };
