Fix shadowed variables

diff --git a/Engine/Source/Runtime/Core/Private/GenericPlatform/GenericPlatformMisc.cpp b/Engine/Source/Runtime/Core/Private/GenericPlatform/GenericPlatformMisc.cpp
--- a/Engine/Source/Runtime/Core/Private/GenericPlatform/GenericPlatformMisc.cpp
+++ b/Engine/Source/Runtime/Core/Private/GenericPlatform/GenericPlatformMisc.cpp
@@ -821,16 +821,16 @@ bool FGenericPlatformMisc::IsRunningOnBattery()
 FGuid FGenericPlatformMisc::GetMachineId()
 {
 	static FGuid MachineId;
-	FString MachineIdStr;
+	FString MachineIdString;
 
 	// Check to see if we already have a valid machine ID to use
-	if( !MachineId.IsValid() && (!FPlatformMisc::GetStoredValue( TEXT( "Epic Games" ), TEXT( "Unreal Engine/Identifiers" ), TEXT( "MachineId" ), MachineIdStr ) || !FGuid::Parse( MachineIdStr, MachineId )) )
+	if( !MachineId.IsValid() && (!FPlatformMisc::GetStoredValue( TEXT( "Epic Games" ), TEXT( "Unreal Engine/Identifiers" ), TEXT( "MachineId" ), MachineIdString ) || !FGuid::Parse( MachineIdString, MachineId )) )
 	{
 		// No valid machine ID, generate and save a new one
 		MachineId = FGuid::NewGuid();
-		MachineIdStr = MachineId.ToString( EGuidFormats::Digits );
+		MachineIdString = MachineId.ToString( EGuidFormats::Digits );
 
-		if( !FPlatformMisc::SetStoredValue( TEXT( "Epic Games" ), TEXT( "Unreal Engine/Identifiers" ), TEXT( "MachineId" ), MachineIdStr ) )
+		if( !FPlatformMisc::SetStoredValue( TEXT( "Epic Games" ), TEXT( "Unreal Engine/Identifiers" ), TEXT( "MachineId" ), MachineIdString ) )
 		{
 			// Failed to persist the machine ID - reset it to zero to avoid returning a transient value
 			MachineId = FGuid();
@@ -842,9 +842,9 @@ FGuid FGenericPlatformMisc::GetMachineId()
 
 FString FGenericPlatformMisc::GetEpicAccountId()
 {
-	FString EpicAccountId;
-	FPlatformMisc::GetStoredValue( TEXT( "Epic Games" ), TEXT( "Unreal Engine/Identifiers" ), TEXT( "AccountId" ), EpicAccountId );
-	return EpicAccountId;
+	FString AccountId;
+	FPlatformMisc::GetStoredValue( TEXT( "Epic Games" ), TEXT( "Unreal Engine/Identifiers" ), TEXT( "AccountId" ), AccountId );
+	return AccountId;
 }
 
 bool FGenericPlatformMisc::SetEpicAccountId( const FString& AccountId )
diff --git a/Engine/Source/Runtime/Online/OnlineSubsystem/Private/OnlineKeyValuePair.cpp b/Engine/Source/Runtime/Online/OnlineSubsystem/Private/OnlineKeyValuePair.cpp
--- a/Engine/Source/Runtime/Online/OnlineSubsystem/Private/OnlineKeyValuePair.cpp
+++ b/Engine/Source/Runtime/Online/OnlineSubsystem/Private/OnlineKeyValuePair.cpp
@@ -468,30 +468,30 @@ TSharedRef<class FJsonObject> FVariantData::ToJson() const
 	{
 		case EOnlineKeyValuePairDataType::Int32:
 		{
-			int32 Value;
-			GetValue(Value);
-			JsonObject->SetNumberField(ValueStr, (double)Value);
+			int32 FieldValue;
+			GetValue(FieldValue);
+			JsonObject->SetNumberField(ValueStr, (double)FieldValue);
 			break;
 		}
 		case EOnlineKeyValuePairDataType::Float:
 		{
-			float Value;
-			GetValue(Value);
-			JsonObject->SetNumberField(ValueStr, (double)Value);
+			float FieldValue;
+			GetValue(FieldValue);
+			JsonObject->SetNumberField(ValueStr, (double)FieldValue);
 			break;
 		}
 		case EOnlineKeyValuePairDataType::String:
 		{
-			FString Value;
-			GetValue(Value);
-			JsonObject->SetStringField(ValueStr, Value);
+			FString FieldValue;
+			GetValue(FieldValue);
+			JsonObject->SetStringField(ValueStr, FieldValue);
 			break;
 		}
 		case EOnlineKeyValuePairDataType::Bool:
 		{
-			bool Value;
-			GetValue(Value);
-			JsonObject->SetBoolField(ValueStr, Value);
+			bool FieldValue;
+			GetValue(FieldValue);
+			JsonObject->SetBoolField(ValueStr, FieldValue);
 			break;
 		}
 		case EOnlineKeyValuePairDataType::Int64:
@@ -501,9 +501,9 @@ TSharedRef<class FJsonObject> FVariantData::ToJson() const
 		}
 		case EOnlineKeyValuePairDataType::Double:
 		{
-			double Value;
-			GetValue(Value);
-			JsonObject->SetNumberField(ValueStr, (double)Value);
+			double FieldValue;
+			GetValue(FieldValue);
+			JsonObject->SetNumberField(ValueStr, (double)FieldValue);
 			break;
 		}
 		case EOnlineKeyValuePairDataType::Empty:
diff --git a/Engine/Source/Runtime/ShaderCore/Public/ShaderCache.h b/Engine/Source/Runtime/ShaderCore/Public/ShaderCache.h
--- a/Engine/Source/Runtime/ShaderCore/Public/ShaderCache.h
+++ b/Engine/Source/Runtime/ShaderCore/Public/ShaderCache.h
@@ -323,7 +323,7 @@ public:
 	
 	static FORCEINLINE void LogShader(EShaderPlatform Platform, EShaderFrequency Frequency, FSHAHash Hash, TArray<uint8> const& Code)
 	{
-		FShaderCache* Cache = GetShaderCache();
+		FShaderCache* ShaderCache = GetShaderCache();
 		if ( Cache )
 		{
 			Cache->InternalLogShader(Platform, Frequency, Hash, Code);
@@ -332,7 +332,7 @@ public:
 	
 	static FORCEINLINE void LogVertexDeclaration(const FVertexDeclarationElementList& VertexElements, FVertexDeclarationRHIParamRef VertexDeclaration)
 	{
-		FShaderCache* Cache = GetShaderCache();
+		FShaderCache* ShaderCache = GetShaderCache();
 		if ( Cache )
 		{
 			Cache->InternalLogVertexDeclaration(VertexElements, VertexDeclaration);
@@ -347,7 +347,7 @@ public:
 												FGeometryShaderRHIParamRef GeometryShader,
 												FBoundShaderStateRHIParamRef BoundState)
 	{
-		FShaderCache* Cache = GetShaderCache();
+		FShaderCache* ShaderCache = GetShaderCache();
 		if ( Cache )
 		{
 			Cache->InternalLogBoundShaderState(Platform, VertexDeclaration, VertexShader, PixelShader, HullShader, DomainShader, GeometryShader, BoundState);
@@ -356,7 +356,7 @@ public:
 	
 	static FORCEINLINE void LogBlendState(FBlendStateInitializerRHI const& Init, FBlendStateRHIParamRef State)
 	{
-		FShaderCache* Cache = GetShaderCache();
+		FShaderCache* ShaderCache = GetShaderCache();
 		if ( Cache )
 		{
 			Cache->InternalLogBlendState(Init, State);
@@ -364,7 +364,7 @@ public:
 	}
 	static FORCEINLINE void LogRasterizerState(FRasterizerStateInitializerRHI const& Init, FRasterizerStateRHIParamRef State)
 	{
-		FShaderCache* Cache = GetShaderCache();
+		FShaderCache* ShaderCache = GetShaderCache();
 		if ( Cache )
 		{
 			Cache->InternalLogRasterizerState(Init, State);
@@ -372,7 +372,7 @@ public:
 	}
 	static FORCEINLINE void LogDepthStencilState(FDepthStencilStateInitializerRHI const& Init, FDepthStencilStateRHIParamRef State)
 	{
-		FShaderCache* Cache = GetShaderCache();
+		FShaderCache* ShaderCache = GetShaderCache();
 		if ( Cache )
 		{
 			Cache->InternalLogDepthStencilState(Init, State);
@@ -380,7 +380,7 @@ public:
 	}
 	static FORCEINLINE void LogSamplerState(FSamplerStateInitializerRHI const& Init, FSamplerStateRHIParamRef State)
 	{
-		FShaderCache* Cache = GetShaderCache();
+		FShaderCache* ShaderCache = GetShaderCache();
 		if ( Cache )
 		{
 			Cache->InternalLogSamplerState(Init, State);
@@ -388,7 +388,7 @@ public:
 	}
 	static FORCEINLINE void LogTexture(FShaderTextureKey const& Init, FTextureRHIParamRef State)
 	{
-		FShaderCache* Cache = GetShaderCache();
+		FShaderCache* ShaderCache = GetShaderCache();
 		if ( Cache )
 		{
 			Cache->InternalLogTexture(Init, State);
@@ -396,7 +396,7 @@ public:
 	}
 	static FORCEINLINE void LogSRV(FShaderResourceViewRHIParamRef SRV, FTextureRHIParamRef Texture, uint8 StartMip, uint8 NumMips, uint8 Format)
 	{
-		FShaderCache* Cache = GetShaderCache();
+		FShaderCache* ShaderCache = GetShaderCache();
 		if ( Cache )
 		{
 			Cache->InternalLogSRV(SRV, Texture, StartMip, NumMips, Format);
@@ -404,7 +404,7 @@ public:
 	}
 	static FORCEINLINE void LogSRV(FShaderResourceViewRHIParamRef SRV, FVertexBufferRHIParamRef Vb, uint32 Stride, uint8 Format)
 	{
-		FShaderCache* Cache = GetShaderCache();
+		FShaderCache* ShaderCache = GetShaderCache();
 		if ( Cache )
 		{
 			Cache->InternalLogSRV(SRV, Vb, Stride, Format);
@@ -412,7 +412,7 @@ public:
 	}
 	static FORCEINLINE void RemoveSRV(FShaderResourceViewRHIParamRef SRV)
 	{
-		FShaderCache* Cache = GetShaderCache();
+		FShaderCache* ShaderCache = GetShaderCache();
 		if ( Cache )
 		{
 			Cache->InternalRemoveSRV(SRV);
@@ -420,7 +420,7 @@ public:
 	}
 	static FORCEINLINE void RemoveTexture(FTextureRHIParamRef Texture)
 	{
-		FShaderCache* Cache = GetShaderCache();
+		FShaderCache* ShaderCache = GetShaderCache();
 		if ( Cache )
 		{
 			Cache->InternalRemoveTexture(Texture);
@@ -428,7 +428,7 @@ public:
 	}
 	static FORCEINLINE void SetBlendState(FBlendStateRHIParamRef State)
 	{
-		FShaderCache* Cache = GetShaderCache();
+		FShaderCache* ShaderCache = GetShaderCache();
 		if ( Cache )
 		{
 			Cache->InternalSetBlendState(State);
@@ -436,7 +436,7 @@ public:
 	}
 	static FORCEINLINE void SetRasterizerState(FRasterizerStateRHIParamRef State)
 	{
-		FShaderCache* Cache = GetShaderCache();
+		FShaderCache* ShaderCache = GetShaderCache();
 		if ( Cache )
 		{
 			Cache->InternalSetRasterizerState(State);
@@ -444,7 +444,7 @@ public:
 	}
 	static FORCEINLINE void SetDepthStencilState(FDepthStencilStateRHIParamRef State)
 	{
-		FShaderCache* Cache = GetShaderCache();
+		FShaderCache* ShaderCache = GetShaderCache();
 		if ( Cache )
 		{
 			Cache->InternalSetDepthStencilState(State);
@@ -452,7 +452,7 @@ public:
 	}
 	static FORCEINLINE void SetRenderTargets( uint32 NumSimultaneousRenderTargets, const FRHIRenderTargetView* NewRenderTargetsRHI, const FRHIDepthRenderTargetView* NewDepthStencilTargetRHI )
 	{
-		FShaderCache* Cache = GetShaderCache();
+		FShaderCache* ShaderCache = GetShaderCache();
 		if ( Cache )
 		{
 			Cache->InternalSetRenderTargets(NumSimultaneousRenderTargets, NewRenderTargetsRHI, NewDepthStencilTargetRHI);
@@ -460,7 +460,7 @@ public:
 	}
 	static FORCEINLINE void SetSamplerState(EShaderFrequency Frequency, uint32 Index, FSamplerStateRHIParamRef State)
 	{
-		FShaderCache* Cache = GetShaderCache();
+		FShaderCache* ShaderCache = GetShaderCache();
 		if ( Cache )
 		{
 			Cache->InternalSetSamplerState(Frequency, Index, State);
@@ -468,7 +468,7 @@ public:
 	}
 	static FORCEINLINE void SetTexture(EShaderFrequency Frequency, uint32 Index, FTextureRHIParamRef State)
 	{
-		FShaderCache* Cache = GetShaderCache();
+		FShaderCache* ShaderCache = GetShaderCache();
 		if ( Cache )
 		{
 			Cache->InternalSetTexture(Frequency, Index, State);
@@ -476,7 +476,7 @@ public:
 	}
 	static FORCEINLINE void SetSRV(EShaderFrequency Frequency, uint32 Index, FShaderResourceViewRHIParamRef SRV)
 	{
-		FShaderCache* Cache = GetShaderCache();
+		FShaderCache* ShaderCache = GetShaderCache();
 		if ( Cache )
 		{
 			Cache->InternalSetSRV(Frequency, Index, SRV);
@@ -484,7 +484,7 @@ public:
 	}
 	static FORCEINLINE void SetBoundShaderState(FBoundShaderStateRHIParamRef State)
 	{
-		FShaderCache* Cache = GetShaderCache();
+		FShaderCache* ShaderCache = GetShaderCache();
 		if ( Cache )
 		{
 			Cache->InternalSetBoundShaderState(State);
@@ -492,7 +492,7 @@ public:
 	}
 	static FORCEINLINE void SetViewport(uint32 MinX, uint32 MinY, float MinZ, uint32 MaxX, uint32 MaxY, float MaxZ)
 	{
-		FShaderCache* Cache = GetShaderCache();
+		FShaderCache* ShaderCache = GetShaderCache();
 		if ( Cache )
 		{
 			Cache->InternalSetViewport(MinX, MinY, MinZ, MaxX, MaxY, MaxZ);
@@ -500,7 +500,7 @@ public:
 	}
 	static FORCEINLINE void LogDraw(uint8 IndexType)
 	{
-		FShaderCache* Cache = GetShaderCache();
+		FShaderCache* ShaderCache = GetShaderCache();
 		if ( Cache )
 		{
 			Cache->InternalLogDraw(IndexType);
@@ -508,7 +508,7 @@ public:
 	}
 	static FORCEINLINE void PreDrawShaders(FRHICommandList& RHICmdList)
 	{
-		FShaderCache* Cache = GetShaderCache();
+		FShaderCache* ShaderCache = GetShaderCache();
 		if ( Cache )
 		{
 			Cache->InternalPreDrawShaders(RHICmdList);
diff --git a/Engine/Source/Runtime/SlateCore/Private/Types/PaintArgs.cpp b/Engine/Source/Runtime/SlateCore/Private/Types/PaintArgs.cpp
--- a/Engine/Source/Runtime/SlateCore/Private/Types/PaintArgs.cpp
+++ b/Engine/Source/Runtime/SlateCore/Private/Types/PaintArgs.cpp
@@ -38,8 +38,8 @@ FPaintArgs FPaintArgs::RecordHittestGeometry(const SWidget* Widget, const FGeome
 	return UpdatedArgs;
 }
 
-FPaintArgs FPaintArgs::InsertCustomHitTestPath( TSharedRef<ICustomHitTestPath> CustomHitTestPath, int32 LastHittestIndex ) const
+FPaintArgs FPaintArgs::InsertCustomHitTestPath( TSharedRef<ICustomHitTestPath> CustomHitTestPath, int32 InLastHittestIndex ) const
 {
-	const_cast<FHittestGrid&>(Grid).InsertCustomHitTestPath( CustomHitTestPath, LastHittestIndex );
+	const_cast<FHittestGrid&>(Grid).InsertCustomHitTestPath( CustomHitTestPath, InLastHittestIndex );
 	return *this;
 }
\ No newline at end of file
diff --git a/Engine/Source/Runtime/SlateCore/Private/Widgets/SWidget.cpp b/Engine/Source/Runtime/SlateCore/Private/Widgets/SWidget.cpp
--- a/Engine/Source/Runtime/SlateCore/Private/Widgets/SWidget.cpp
+++ b/Engine/Source/Runtime/SlateCore/Private/Widgets/SWidget.cpp
@@ -327,11 +327,11 @@ void SWidget::OnFinishedKeyInput()
 FNavigationReply SWidget::OnNavigation(const FGeometry& MyGeometry, const FNavigationEvent& InNavigationEvent)
 {
 	EUINavigation Type = InNavigationEvent.GetNavigationType();
-	TSharedPtr<FNavigationMetaData> MetaData = GetMetaData<FNavigationMetaData>();
-	if (MetaData.IsValid())
+	TSharedPtr<FNavigationMetaData> NavigationMetaData = GetMetaData<FNavigationMetaData>();
+	if (NavigationMetaData.IsValid())
 	{
-		TSharedPtr<SWidget> Widget = MetaData->GetFocusRecipient(Type).Pin();
-		return FNavigationReply(MetaData->GetBoundaryRule(Type), Widget, MetaData->GetFocusDelegate(Type));
+		TSharedPtr<SWidget> Widget = NavigationMetaData->GetFocusRecipient(Type).Pin();
+		return FNavigationReply(NavigationMetaData->GetBoundaryRule(Type), Widget, NavigationMetaData->GetFocusDelegate(Type));
 	}
 	return FNavigationReply::Escape();
 }
diff --git a/Engine/Source/Runtime/SlateCore/Private/Widgets/SWindow.cpp b/Engine/Source/Runtime/SlateCore/Private/Widgets/SWindow.cpp
--- a/Engine/Source/Runtime/SlateCore/Private/Widgets/SWindow.cpp
+++ b/Engine/Source/Runtime/SlateCore/Private/Widgets/SWindow.cpp
@@ -608,8 +608,8 @@ FGeometry SWindow::GetWindowGeometryInWindow() const
 	// FGeometry expects Size in Local space, but our size is stored in screen space (same as window space + screen offset).
 	// So we need to transform Size into the window's local space for FGeometry.
 	FSlateLayoutTransform LocalToWindow = GetLocalToWindowTransform();
-	FVector2D Size = GetViewportSize();
-	return FGeometry::MakeRoot( TransformVector(Inverse(LocalToWindow), Size), LocalToWindow );
+	FVector2D ViewportSize = GetViewportSize();
+	return FGeometry::MakeRoot( TransformVector(Inverse(LocalToWindow), ViewportSize), LocalToWindow );
 }
 
 FSlateLayoutTransform SWindow::GetLocalToScreenTransform() const
@@ -674,8 +674,8 @@ FVector2D SWindow::GetClientSizeInScreen() const
 
 FSlateRect SWindow::GetClippingRectangleInWindow() const
 {
-	FVector2D Size = GetViewportSize();
-	return FSlateRect( 0, 0, Size.X, Size.Y );
+	FVector2D ViewportSize = GetViewportSize();
+	return FSlateRect( 0, 0, ViewportSize.X, ViewportSize.Y );
 }
 
 
