Fixed sorting issue with vertex factory when serializing.  Previously did not consider strings with different lengths different.

diff --git a/Engine/Source/Runtime/ShaderCore/Public/VertexFactory.h b/Engine/Source/Runtime/ShaderCore/Public/VertexFactory.h
--- a/Engine/Source/Runtime/ShaderCore/Public/VertexFactory.h
+++ b/Engine/Source/Runtime/ShaderCore/Public/VertexFactory.h
@@ -314,7 +314,13 @@ class FCompareVertexFactoryTypes
 public:		
 	FORCEINLINE bool operator()(const FVertexFactoryType& A, const FVertexFactoryType& B ) const
 	{
-		return FCString::Strncmp(A.GetName(), B.GetName(), FMath::Min(FCString::Strlen(A.GetName()), FCString::Strlen(B.GetName()))) > 0;
+		int32 AL = FCString::Strlen(A.GetName());
+		int32 BL = FCString::Strlen(B.GetName());
+		if ( AL == BL )
+		{
+			return FCString::Strncmp(A.GetName(), B.GetName(), AL) > 0;
+		}
+		return AL > BL;
 	}
 };
 
