UE4 - UE-12059 - Improve editor startup time

diff --git a/Engine/Source/Runtime/ShaderCore/Private/ShaderCore.cpp b/Engine/Source/Runtime/ShaderCore/Private/ShaderCore.cpp
--- a/Engine/Source/Runtime/ShaderCore/Private/ShaderCore.cpp
+++ b/Engine/Source/Runtime/ShaderCore/Private/ShaderCore.cpp
@@ -514,21 +514,42 @@ void BuildShaderFileToUniformBufferMap(TMap<FString, TArray<const TCHAR*> >& Sha
 
 		FScopedSlowTask SlowTask(ShaderSourceFiles.Num());
 
+		// Cache UB access strings, make it case sensitive for faster search
+		struct FShaderVariable
+		{
+			FShaderVariable(const TCHAR* ShaderVariable) :
+				OriginalShaderVariable(ShaderVariable), SearchKey(FString(ShaderVariable).ToUpper() + TEXT("."))
+			{
+			}
+			const TCHAR* OriginalShaderVariable;
+			FString SearchKey;
+		};
+		// Cache each UB
+		TArray<FShaderVariable> SearchKeys;
+		for (TLinkedList<FUniformBufferStruct*>::TIterator StructIt(FUniformBufferStruct::GetStructList()); StructIt; StructIt.Next())
+		{
+			SearchKeys.Add(FShaderVariable(StructIt->GetShaderVariableName()));
+		}
+
+		// Find for each shader file which UBs it needs
 		for (int32 FileIndex = 0; FileIndex < ShaderSourceFiles.Num(); FileIndex++)
 		{
 			SlowTask.EnterProgressFrame(1);
-			FString ShaderFileContents;
 
+ 			FString ShaderFileContents;
 			LoadShaderSourceFileChecked(*ShaderSourceFiles[FileIndex], ShaderFileContents);
 
+			// To allow case sensitive search which is way faster on some platforms (no need to look up locale, etc)
+			ShaderFileContents = ShaderFileContents.ToUpper();
+
 			TArray<const TCHAR*>& ReferencedUniformBuffers = ShaderFileToUniformBufferVariables.FindOrAdd(ShaderSourceFiles[FileIndex]);
 
-			for (TLinkedList<FUniformBufferStruct*>::TIterator StructIt(FUniformBufferStruct::GetStructList()); StructIt; StructIt.Next())
+			for (int32 SearchKeyIndex = 0; SearchKeyIndex < SearchKeys.Num(); ++SearchKeyIndex)
 			{
 				// Searching for the uniform buffer shader variable being accessed with '.'
-				if (ShaderFileContents.Contains(FString(StructIt->GetShaderVariableName()) + TEXT(".")))
+				if (ShaderFileContents.Contains(SearchKeys[SearchKeyIndex].SearchKey, ESearchCase::CaseSensitive))
 				{
-					ReferencedUniformBuffers.AddUnique(StructIt->GetShaderVariableName());
+					ReferencedUniformBuffers.AddUnique(SearchKeys[SearchKeyIndex].OriginalShaderVariable);
 				}
 			}
 		}
