UE4 - Get more info out of not reproducible crash (UE-11434)

diff --git a/Engine/Source/Runtime/ShaderCore/Private/Shader.cpp b/Engine/Source/Runtime/ShaderCore/Private/Shader.cpp
--- a/Engine/Source/Runtime/ShaderCore/Private/Shader.cpp
+++ b/Engine/Source/Runtime/ShaderCore/Private/Shader.cpp
@@ -201,25 +201,29 @@ FArchive& operator<<(FArchive& Ar,FShaderType*& Ref)
 {
 	if(Ar.IsSaving())
 	{
-		FName FactoryName = Ref ? FName(Ref->Name) : NAME_None;
-		Ar << FactoryName;
+		FName ShaderTypeName = Ref ? FName(Ref->Name) : NAME_None;
+		Ar << ShaderTypeName;
 	}
 	else if(Ar.IsLoading())
 	{
-		FName FactoryName = NAME_None;
-		Ar << FactoryName;
+		FName ShaderTypeName = NAME_None;
+		Ar << ShaderTypeName;
 		
 		Ref = NULL;
 
-		if(FactoryName != NAME_None)
+		if(ShaderTypeName != NAME_None)
 		{
 			// look for the shader type in the global name to type map
-			FShaderType** ShaderType = FShaderType::GetNameToTypeMap().Find(FactoryName);
+			FShaderType** ShaderType = FShaderType::GetNameToTypeMap().Find(ShaderTypeName);
 			if (ShaderType)
 			{
 				// if we found it, use it
 				Ref = *ShaderType;
 			}
+			else
+			{
+				UE_LOG(LogShaders, Warning, TEXT("ShaderType '%s' was not found!"), *ShaderTypeName.ToString());
+			}
 		}
 	}
 	return Ar;
@@ -755,12 +759,11 @@ FShader::FShader(const CompiledShaderInitializerType& Initializer):
 	OutputHash = Initializer.OutputHash;
 	checkSlow(OutputHash != FSHAHash());
 
-	if (Type)
-	{
-		// Store off the source hash that this shader was compiled with
-		// This will be used as part of the shader key in order to identify when shader files have been changed and a recompile is needed
-		SourceHash = Type->GetSourceHash();
-	}
+	check(Type);
+
+	// Store off the source hash that this shader was compiled with
+	// This will be used as part of the shader key in order to identify when shader files have been changed and a recompile is needed
+	SourceHash = Type->GetSourceHash();
 
 	if (VFType)
 	{
