Prevent the shader-cache from instantiating incompatible shader resources & crashing the RHI, requires exposing the existing ArePlatformsCompatible function from FShaderResource.

diff --git a/Engine/Source/Runtime/ShaderCore/Private/Shader.cpp b/Engine/Source/Runtime/ShaderCore/Private/Shader.cpp
--- a/Engine/Source/Runtime/ShaderCore/Private/Shader.cpp
+++ b/Engine/Source/Runtime/ShaderCore/Private/Shader.cpp
@@ -404,7 +404,7 @@ void FShaderResource::FinishCleanup()
 	delete this;
 }
 
-bool ArePlatformsCompatible(EShaderPlatform CurrentPlatform, EShaderPlatform TargetPlatform)
+bool FShaderResource::ArePlatformsCompatible(EShaderPlatform CurrentPlatform, EShaderPlatform TargetPlatform)
 {
 	bool bFeatureLevelCompatible = CurrentPlatform == TargetPlatform;
 	
diff --git a/Engine/Source/Runtime/ShaderCore/Private/ShaderCache.cpp b/Engine/Source/Runtime/ShaderCore/Private/ShaderCache.cpp
--- a/Engine/Source/Runtime/ShaderCore/Private/ShaderCache.cpp
+++ b/Engine/Source/Runtime/ShaderCore/Private/ShaderCache.cpp
@@ -308,14 +308,37 @@ FComputeShaderRHIRef FShaderCache::GetComputeShader(EShaderPlatform Platform, TA
 
 void FShaderCache::InternalLogShader(EShaderPlatform Platform, EShaderFrequency Frequency, FSHAHash Hash, TArray<uint8> const& Code)
 {
-	FShaderCacheKey Key;
-	Key.Hash = Hash;
-	Key.Platform = Platform;
-	Key.Frequency = Frequency;
-	Key.bActive = true;
-	ENQUEUE_UNIQUE_RENDER_COMMAND_TWOPARAMETER(LogShader, FShaderCacheKey,Key,Key, TArray<uint8>,Code,Code, {
-		FShaderCache::GetShaderCache()->SubmitShader(Key, Code);
-	});
+	bool bUsable = FShaderResource::ArePlatformsCompatible(GMaxRHIShaderPlatform, Platform);
+	switch (Frequency)
+	{
+		case SF_Geometry:
+			bUsable &= RHISupportsGeometryShaders(GMaxRHIShaderPlatform);
+			break;
+			
+		case SF_Hull:
+		case SF_Domain:
+			bUsable &= RHISupportsTessellation(GMaxRHIShaderPlatform);
+			break;
+			
+		case SF_Compute:
+			bUsable &= RHISupportsComputeShaders(GMaxRHIShaderPlatform);
+			break;
+			
+		default:
+			break;
+	}
+	
+	if ( bUsable )
+	{
+		FShaderCacheKey Key;
+		Key.Hash = Hash;
+		Key.Platform = Platform;
+		Key.Frequency = Frequency;
+		Key.bActive = true;
+		ENQUEUE_UNIQUE_RENDER_COMMAND_TWOPARAMETER(LogShader, FShaderCacheKey,Key,Key, TArray<uint8>,Code,Code, {
+			FShaderCache::GetShaderCache()->SubmitShader(Key, Code);
+		});
+	}
 }
 
 void FShaderCache::InternalLogVertexDeclaration(const FVertexDeclarationElementList& VertexElements, FVertexDeclarationRHIParamRef VertexDeclaration)
diff --git a/Engine/Source/Runtime/ShaderCore/Public/Shader.h b/Engine/Source/Runtime/ShaderCore/Public/Shader.h
--- a/Engine/Source/Runtime/ShaderCore/Public/Shader.h
+++ b/Engine/Source/Runtime/ShaderCore/Public/Shader.h
@@ -171,6 +171,9 @@ public:
 	/** Return a list of all shader Ids currently known */
 	SHADERCORE_API static void GetAllShaderResourceId(TArray<FShaderResourceId>& Ids);
 
+	/** Returns true if and only if TargetPlatform is compatible for use with CurrentPlatform. */
+	SHADERCORE_API static bool ArePlatformsCompatible(EShaderPlatform CurrentPlatform, EShaderPlatform TargetPlatform);
+	
 private:
 
 	/** Reference to the RHI shader.  Only one of these is ever valid, and it is the one corresponding to Target.Frequency. */
