Disable the ShaderCache's draw-state logging by default in all circumstances since it is currently a significant performance impact and not all project use Shipping builds for release.

diff --git a/Engine/Source/Runtime/ShaderCore/Private/ShaderCache.cpp b/Engine/Source/Runtime/ShaderCore/Private/ShaderCache.cpp
--- a/Engine/Source/Runtime/ShaderCore/Private/ShaderCache.cpp
+++ b/Engine/Source/Runtime/ShaderCore/Private/ShaderCache.cpp
@@ -33,8 +33,8 @@ FAutoConsoleVariableRef FShaderCache::CVarUseShaderPredraw(
 	ECVF_ReadOnly|ECVF_RenderThreadSafe
 	);
 
-// The actual draw loggging is even more expensive as it has to cache all the RHI draw state & is disabled by default if you aren't a developer.
-int32 FShaderCache::bUseShaderDrawLog = (PLATFORM_MAC && !WITH_EDITOR && !UE_BUILD_SHIPPING) ? 1 : 0;
+// The actual draw loggging is even more expensive as it has to cache all the RHI draw state & is disabled by default.
+int32 FShaderCache::bUseShaderDrawLog = 0;
 FAutoConsoleVariableRef FShaderCache::CVarUseShaderDrawLog(
 	TEXT("r.UseShaderDrawLog"),
 	bUseShaderDrawLog,
