Fix a bug in early submission of compute shaders in the ShaderCache not using the correct hash.

diff --git a/Engine/Source/Runtime/ShaderCore/Private/ShaderCache.cpp b/Engine/Source/Runtime/ShaderCore/Private/ShaderCache.cpp
--- a/Engine/Source/Runtime/ShaderCore/Private/ShaderCache.cpp
+++ b/Engine/Source/Runtime/ShaderCore/Private/ShaderCache.cpp
@@ -941,10 +941,11 @@ void FShaderCache::SubmitShader(FShaderCacheKey const& Key, TArray<uint8> const&
 				if (IsValidRef(Shader))
 				{
 					// @todo WARNING: The RHI is responsible for hashing Compute shaders, unlike other stages because of how OpenGLDrv implements compute.
-					// Shader->SetHash(Key.Hash);
-					CachedComputeShaders.Add(Key, Shader);
-					PrebindShader(Key);
-					PlatformCache.Shaders.Add(Key);
+					FShaderCacheKey ComputeKey = Key;
+					ComputeKey.Hash = Shader->GetHash();
+					CachedComputeShaders.Add(ComputeKey, Shader);
+					PrebindShader(ComputeKey);
+					PlatformCache.Shaders.Add(ComputeKey);
 				}
 			}
 			break;
