In the shader-cache precompute the object hashes wherever possible to reduce the time spent searching maps, serialising the hashes to boot. Increments the current cache format version.

diff --git a/Engine/Source/Runtime/ShaderCore/Private/ShaderCache.cpp b/Engine/Source/Runtime/ShaderCore/Private/ShaderCache.cpp
--- a/Engine/Source/Runtime/ShaderCore/Private/ShaderCache.cpp
+++ b/Engine/Source/Runtime/ShaderCore/Private/ShaderCache.cpp
@@ -185,7 +185,7 @@ FVertexShaderRHIRef FShaderCache::GetVertexShader(EShaderPlatform Platform, FSHA
 	FShaderCacheKey Key;
 	Key.Platform = Platform;
 	Key.Frequency = SF_Vertex;
-	Key.Hash = Hash;
+	Key.SHAHash = Hash;
 	Key.bActive = true;
 	FVertexShaderRHIRef Shader = CachedVertexShaders.FindRef(Key);
 	if(!IsValidRef(Shader))
@@ -206,7 +206,7 @@ FPixelShaderRHIRef FShaderCache::GetPixelShader(EShaderPlatform Platform, FSHAHa
 	FShaderCacheKey Key;
 	Key.Platform = Platform;
 	Key.Frequency = SF_Pixel;
-	Key.Hash = Hash;
+	Key.SHAHash = Hash;
 	Key.bActive = true;
 	FPixelShaderRHIRef Shader = CachedPixelShaders.FindRef(Key);
 	if(!IsValidRef(Shader))
@@ -227,7 +227,7 @@ FGeometryShaderRHIRef FShaderCache::GetGeometryShader(EShaderPlatform Platform,
 	FShaderCacheKey Key;
 	Key.Platform = Platform;
 	Key.Frequency = SF_Geometry;
-	Key.Hash = Hash;
+	Key.SHAHash = Hash;
 	Key.bActive = true;
 	FGeometryShaderRHIRef Shader = CachedGeometryShaders.FindRef(Key);
 	if(!IsValidRef(Shader))
@@ -248,7 +248,7 @@ FHullShaderRHIRef FShaderCache::GetHullShader(EShaderPlatform Platform, FSHAHash
 	FShaderCacheKey Key;
 	Key.Platform = Platform;
 	Key.Frequency = SF_Hull;
-	Key.Hash = Hash;
+	Key.SHAHash = Hash;
 	Key.bActive = true;
 	FHullShaderRHIRef Shader = CachedHullShaders.FindRef(Key);
 	if(!IsValidRef(Shader))
@@ -269,7 +269,7 @@ FDomainShaderRHIRef FShaderCache::GetDomainShader(EShaderPlatform Platform, FSHA
 	FShaderCacheKey Key;
 	Key.Platform = Platform;
 	Key.Frequency = SF_Domain;
-	Key.Hash = Hash;
+	Key.SHAHash = Hash;
 	Key.bActive = true;
 	FDomainShaderRHIRef Shader = CachedDomainShaders.FindRef(Key);
 	if(!IsValidRef(Shader))
@@ -291,7 +291,7 @@ FComputeShaderRHIRef FShaderCache::GetComputeShader(EShaderPlatform Platform, TA
 	Key.Platform = Platform;
 	Key.Frequency = SF_Compute;
 	// @todo WARNING: RHI is responsible for hashing Compute shaders due to the way OpenGLDrv implements compute!
-	FSHA1::HashBuffer(Code.GetData(), Code.Num(), Key.Hash.Hash);
+	FSHA1::HashBuffer(Code.GetData(), Code.Num(), Key.SHAHash.Hash);
 	Key.bActive = true;
 	FComputeShaderRHIRef Shader = CachedComputeShaders.FindRef(Key);
 	if(!IsValidRef(Shader))
@@ -331,7 +331,7 @@ void FShaderCache::InternalLogShader(EShaderPlatform Platform, EShaderFrequency
 	if ( bUsable )
 	{
 		FShaderCacheKey Key;
-		Key.Hash = Hash;
+		Key.SHAHash = Hash;
 		Key.Platform = Platform;
 		Key.Frequency = Frequency;
 		Key.bActive = true;
@@ -363,35 +363,35 @@ void FShaderCache::InternalLogBoundShaderState(EShaderPlatform Platform, FVertex
 	{
 		Info.VertexShader.Platform = Platform;
 		Info.VertexShader.Frequency = SF_Vertex;
-		Info.VertexShader.Hash = VertexShader->GetHash();
+		Info.VertexShader.SHAHash = VertexShader->GetHash();
 		Info.VertexShader.bActive = true;
 	}
 	if(PixelShader)
 	{
 		Info.PixelShader.Platform = Platform;
 		Info.PixelShader.Frequency = SF_Pixel;
-		Info.PixelShader.Hash = PixelShader->GetHash();
+		Info.PixelShader.SHAHash = PixelShader->GetHash();
 		Info.PixelShader.bActive = true;
 	}
 	if(GeometryShader)
 	{
 		Info.GeometryShader.Platform = Platform;
 		Info.GeometryShader.Frequency = SF_Geometry;
-		Info.GeometryShader.Hash = GeometryShader->GetHash();
+		Info.GeometryShader.SHAHash = GeometryShader->GetHash();
 		Info.GeometryShader.bActive = true;
 	}
 	if(HullShader)
 	{
 		Info.HullShader.Platform = Platform;
 		Info.HullShader.Frequency = SF_Hull;
-		Info.HullShader.Hash = HullShader->GetHash();
+		Info.HullShader.SHAHash = HullShader->GetHash();
 		Info.HullShader.bActive = true;
 	}
 	if(DomainShader)
 	{
 		Info.DomainShader.Platform = Platform;
 		Info.DomainShader.Frequency = SF_Domain;
-		Info.DomainShader.Hash = DomainShader->GetHash();
+		Info.DomainShader.SHAHash = DomainShader->GetHash();
 		Info.DomainShader.bActive = true;
 	}
 	
@@ -556,6 +556,7 @@ void FShaderCache::InternalSetBlendState(FBlendStateRHIParamRef State)
 	if ( (bUseShaderPredraw || bUseShaderDrawLog) && !bIsPreDraw )
 	{
 		CurrentDrawKey.BlendState = BlendStates.FindChecked(State);
+		CurrentDrawKey.Hash = 0;
 	}
 }
 
@@ -564,6 +565,7 @@ void FShaderCache::InternalSetRasterizerState(FRasterizerStateRHIParamRef State)
 	if ( (bUseShaderPredraw || bUseShaderDrawLog) && !bIsPreDraw )
 	{
 		CurrentDrawKey.RasterizerState = RasterizerStates.FindChecked(State);
+		CurrentDrawKey.Hash = 0;
 	}
 }
 
@@ -572,6 +574,7 @@ void FShaderCache::InternalSetDepthStencilState(FDepthStencilStateRHIParamRef St
 	if ( (bUseShaderPredraw || bUseShaderDrawLog) && !bIsPreDraw )
 	{
 		CurrentDrawKey.DepthStencilState = DepthStencilStates.FindChecked(State);
+		CurrentDrawKey.Hash = 0;
 	}
 }
 
@@ -619,6 +622,7 @@ void FShaderCache::InternalSetRenderTargets( uint32 NumSimultaneousRenderTargets
 		{
 			CurrentDrawKey.DepthStencilTarget = FShaderDrawKey::NullState;
 		}
+		CurrentDrawKey.Hash = 0;
 	}
 }
 
@@ -638,6 +642,7 @@ void FShaderCache::InternalSetSamplerState(EShaderFrequency Frequency, uint32 In
 		{
 			CurrentDrawKey.SamplerStates[Frequency][Index] = FShaderDrawKey::NullState;
 		}
+		CurrentDrawKey.Hash = 0;
 	}
 }
 
@@ -666,6 +671,7 @@ void FShaderCache::InternalSetTexture(EShaderFrequency Frequency, uint32 Index,
 		{
 			CurrentDrawKey.Resources[Frequency][Index] = FShaderDrawKey::NullState;
 		}
+		CurrentDrawKey.Hash = 0;
 	}
 }
 
@@ -685,6 +691,7 @@ void FShaderCache::InternalSetSRV(EShaderFrequency Frequency, uint32 Index, FSha
 		{
 			CurrentDrawKey.Resources[Frequency][Index] = FShaderDrawKey::NullState;
 		}
+		CurrentDrawKey.Hash = 0;
 	}
 }
 
@@ -699,6 +706,7 @@ void FShaderCache::InternalSetBoundShaderState(FBoundShaderStateRHIParamRef Stat
 		{
 			BoundShaderState = ShaderStates.FindChecked(State);
 		}
+		CurrentDrawKey.Hash = 0;
 	}
 }
 
@@ -873,7 +881,7 @@ void FShaderCache::SubmitShader(FShaderCacheKey const& Key, TArray<uint8> const&
 				FVertexShaderRHIRef Shader = RHICreateVertexShader(Code);
 				if(IsValidRef(Shader))
 				{
-					Shader->SetHash(Key.Hash);
+					Shader->SetHash(Key.SHAHash);
 					CachedVertexShaders.Add(Key, Shader);
 					PrebindShader(Key);
 					PlatformCache.Shaders.Add(Key);
@@ -886,7 +894,7 @@ void FShaderCache::SubmitShader(FShaderCacheKey const& Key, TArray<uint8> const&
 				FPixelShaderRHIRef Shader = RHICreatePixelShader(Code);
 				if(IsValidRef(Shader))
 				{
-					Shader->SetHash(Key.Hash);
+					Shader->SetHash(Key.SHAHash);
 					CachedPixelShaders.Add(Key, Shader);
 					PrebindShader(Key);
 					PlatformCache.Shaders.Add(Key);
@@ -899,7 +907,7 @@ void FShaderCache::SubmitShader(FShaderCacheKey const& Key, TArray<uint8> const&
 				FGeometryShaderRHIRef Shader = RHICreateGeometryShader(Code);
 				if(IsValidRef(Shader))
 				{
-					Shader->SetHash(Key.Hash);
+					Shader->SetHash(Key.SHAHash);
 					CachedGeometryShaders.Add(Key, Shader);
 					PrebindShader(Key);
 					PlatformCache.Shaders.Add(Key);
@@ -912,7 +920,7 @@ void FShaderCache::SubmitShader(FShaderCacheKey const& Key, TArray<uint8> const&
 				FHullShaderRHIRef Shader = RHICreateHullShader(Code);
 				if(IsValidRef(Shader))
 				{
-					Shader->SetHash(Key.Hash);
+					Shader->SetHash(Key.SHAHash);
 					CachedHullShaders.Add(Key, Shader);
 					PrebindShader(Key);
 					PlatformCache.Shaders.Add(Key);
@@ -925,7 +933,7 @@ void FShaderCache::SubmitShader(FShaderCacheKey const& Key, TArray<uint8> const&
 				FDomainShaderRHIRef Shader = RHICreateDomainShader(Code);
 				if(IsValidRef(Shader))
 				{
-					Shader->SetHash(Key.Hash);
+					Shader->SetHash(Key.SHAHash);
 					CachedDomainShaders.Add(Key, Shader);
 					PrebindShader(Key);
 					PlatformCache.Shaders.Add(Key);
@@ -942,7 +950,7 @@ void FShaderCache::SubmitShader(FShaderCacheKey const& Key, TArray<uint8> const&
 				{
 					// @todo WARNING: The RHI is responsible for hashing Compute shaders, unlike other stages because of how OpenGLDrv implements compute.
 					FShaderCacheKey ComputeKey = Key;
-					ComputeKey.Hash = Shader->GetHash();
+					ComputeKey.SHAHash = Shader->GetHash();
 					CachedComputeShaders.Add(ComputeKey, Shader);
 					PrebindShader(ComputeKey);
 					PlatformCache.Shaders.Add(ComputeKey);
diff --git a/Engine/Source/Runtime/ShaderCore/Public/ShaderCache.h b/Engine/Source/Runtime/ShaderCore/Public/ShaderCache.h
--- a/Engine/Source/Runtime/ShaderCore/Public/ShaderCache.h
+++ b/Engine/Source/Runtime/ShaderCore/Public/ShaderCache.h
@@ -15,7 +15,7 @@ struct SHADERCORE_API FShaderCacheCustomVersion
 {
 	static const FGuid Key;
 	static const FGuid GameKey;
-	enum Type {	Initial, PreDraw, Latest = PreDraw };
+	enum Type {	Initial, PreDraw, CacheHashes, Latest = CacheHashes };
 };
 
 enum EShaderCacheTextureType
@@ -33,6 +33,7 @@ enum EShaderCacheTextureType
 
 struct SHADERCORE_API FShaderTextureKey
 {
+	mutable uint32 Hash;
 	uint32 X;
 	uint32 Y;
 	uint32 Z;
@@ -41,9 +42,11 @@ struct SHADERCORE_API FShaderTextureKey
 	uint32 Samples;
 	uint8 Format;
 	TEnumAsByte<EShaderCacheTextureType> Type;
+
 	
 	FShaderTextureKey()
-	: X(0)
+	: Hash(0)
+	, X(0)
 	, Y(0)
 	, Z(0)
 	, Flags(0)
@@ -61,33 +64,38 @@ struct SHADERCORE_API FShaderTextureKey
 	
 	friend uint32 GetTypeHash(const FShaderTextureKey &Key)
 	{
-		uint32 Hash = Key.X;
-		Hash ^= Key.Y;
-		Hash ^= Key.Z;
-		Hash ^= Key.Flags;
-		Hash ^= Key.MipLevels;
-		Hash ^= Key.Samples;
-		Hash ^= Key.Format;
-		Hash ^= Key.Type;
-		return Hash;
+		if(!Key.Hash)
+		{
+			Key.Hash = Key.X;
+			Key.Hash ^= Key.Y;
+			Key.Hash ^= Key.Z;
+			Key.Hash ^= Key.Flags;
+			Key.Hash ^= Key.MipLevels;
+			Key.Hash ^= Key.Samples;
+			Key.Hash ^= Key.Format;
+			Key.Hash ^= Key.Type;
+		}
+		return Key.Hash;
 	}
 		
 	friend FArchive& operator<<( FArchive& Ar, FShaderTextureKey& Info )
 	{
-		return Ar << Info.Format << Info.Type << Info.Samples << Info.MipLevels << Info.Flags << Info.X << Info.Y << Info.Z;
+		return Ar << Info.Format << Info.Type << Info.Samples << Info.MipLevels << Info.Flags << Info.X << Info.Y << Info.Z << Info.Hash;
 	}
 };
 
 struct SHADERCORE_API FShaderResourceKey
 {
 	FShaderTextureKey Tex;
+	mutable uint32 Hash;
 	uint32 BaseMip;
 	uint32 MipLevels;
 	uint8 Format;
 	bool bSRV;
 	
 	FShaderResourceKey()
-	: BaseMip(0)
+	: Hash(0)
+	, BaseMip(0)
 	, MipLevels(0)
 	, Format(0)
 	, bSRV(false)
@@ -101,28 +109,33 @@ struct SHADERCORE_API FShaderResourceKey
 	
 	friend uint32 GetTypeHash(const FShaderResourceKey &Key)
 	{
-		uint32 Hash = GetTypeHash(Key.Tex);
-		Hash ^= Key.BaseMip;
-		Hash ^= Key.MipLevels;
-		Hash ^= Key.Format;
-		Hash ^= (uint32)Key.bSRV;
-		return Hash;
+		if(!Key.Hash)
+		{
+			Key.Hash = GetTypeHash(Key.Tex);
+			Key.Hash ^= Key.BaseMip;
+			Key.Hash ^= Key.MipLevels;
+			Key.Hash ^= Key.Format;
+			Key.Hash ^= (uint32)Key.bSRV;
+		}
+		return Key.Hash;
 	}
 	
 	friend FArchive& operator<<( FArchive& Ar, FShaderResourceKey& Info )
 	{
-		return Ar << Info.Tex << Info.BaseMip << Info.MipLevels << Info.Format << Info.bSRV;
+		return Ar << Info.Tex << Info.BaseMip << Info.MipLevels << Info.Format << Info.bSRV << Info.Hash;
 	}
 };
 
 struct SHADERCORE_API FShaderRenderTargetKey
 {
 	FShaderTextureKey Texture;
+	mutable uint32 Hash;
 	uint32 MipLevel;
 	uint32 ArrayIndex;
 	
 	FShaderRenderTargetKey()
-	: MipLevel(0)
+	: Hash(0)
+	, MipLevel(0)
 	, ArrayIndex(0)
 	{
 		
@@ -135,15 +148,18 @@ struct SHADERCORE_API FShaderRenderTargetKey
 	
 	friend uint32 GetTypeHash(const FShaderRenderTargetKey &Key)
 	{
-		uint32 Hash = GetTypeHash(Key.Texture);
-		Hash ^= Key.MipLevel;
-		Hash ^= Key.ArrayIndex;
-		return Hash;
+		if(!Key.Hash)
+		{
+			Key.Hash = GetTypeHash(Key.Texture);
+			Key.Hash ^= Key.MipLevel;
+			Key.Hash ^= Key.ArrayIndex;
+		}
+		return Key.Hash;
 	}
 	
 	friend FArchive& operator<<( FArchive& Ar, FShaderRenderTargetKey& Info )
 	{
-		return Ar << Info.Texture << Info.MipLevel << Info.ArrayIndex;
+		return Ar << Info.Texture << Info.MipLevel << Info.ArrayIndex << Info.Hash;
 	}
 };
 
@@ -151,23 +167,28 @@ class SHADERCORE_API FShaderCache
 {
 	struct SHADERCORE_API FShaderCacheKey
 	{
-		FShaderCacheKey() : Platform(SP_NumPlatforms), Frequency(SF_NumFrequencies), bActive(false) {}
+		FShaderCacheKey() : Platform(SP_NumPlatforms), Frequency(SF_NumFrequencies), Hash(0), bActive(false) {}
 	
-		FSHAHash Hash;
+		FSHAHash SHAHash;
 		EShaderPlatform Platform;
 		EShaderFrequency Frequency;
+		mutable uint32 Hash;
 		bool bActive;
 		
 		friend bool operator ==(const FShaderCacheKey& A,const FShaderCacheKey& B)
 		{
-			return A.Hash == B.Hash && A.Platform == B.Platform && A.Frequency == B.Frequency && A.bActive == B.bActive;
+			return A.SHAHash == B.SHAHash && A.Platform == B.Platform && A.Frequency == B.Frequency && A.bActive == B.bActive;
 		}
 
 		friend uint32 GetTypeHash(const FShaderCacheKey &Key)
 		{
-			uint32 TargetFrequency = Key.Frequency;
-			uint32 TargetPlatform = Key.Platform;
-			return FCrc::MemCrc_DEPRECATED((const void*)&Key.Hash, sizeof(Key.Hash)) ^ GetTypeHash(TargetPlatform) ^ GetTypeHash(TargetFrequency) ^ GetTypeHash(Key.bActive);
+			if(!Key.Hash)
+			{
+				uint32 TargetFrequency = Key.Frequency;
+				uint32 TargetPlatform = Key.Platform;
+				Key.Hash = FCrc::MemCrc_DEPRECATED((const void*)&Key.SHAHash, sizeof(Key.SHAHash)) ^ GetTypeHash(TargetPlatform) ^ GetTypeHash(TargetFrequency) ^ GetTypeHash(Key.bActive);
+			}
+			return Key.Hash;
 		}
 		
 		friend FArchive& operator<<( FArchive& Ar, FShaderCacheKey& Info )
@@ -177,7 +198,7 @@ class SHADERCORE_API FShaderCache
 			Ar << TargetFrequency << TargetPlatform;
 			Info.Frequency = (EShaderFrequency)TargetFrequency;
 			Info.Platform = (EShaderPlatform)TargetPlatform;
-			return Ar << Info.Hash << Info.bActive;
+			return Ar << Info.SHAHash << Info.bActive << Info.Hash;
 		}
 	};
 	
@@ -191,6 +212,7 @@ class SHADERCORE_API FShaderCache
 		
 		FShaderDrawKey()
 		: DepthStencilTarget(NullState)
+		, Hash(0)
 		, IndexType(0)
 		{
 			FMemory::Memzero(&BlendState, sizeof(BlendState));
@@ -250,6 +272,7 @@ class SHADERCORE_API FShaderCache
 		uint32 SamplerStates[SF_NumFrequencies][MaxTextureSamplers];
 		uint32 Resources[SF_NumFrequencies][MaxTextureSamplers];
 		uint32 DepthStencilTarget;
+		mutable uint32 Hash;
 		uint8 IndexType;
 		
 		friend bool operator ==(const FShaderDrawKey& A,const FShaderDrawKey& B)
@@ -267,15 +290,18 @@ class SHADERCORE_API FShaderCache
 		
 		friend uint32 GetTypeHash(const FShaderDrawKey &Key)
 		{
-			uint32 CRC = FCrc::MemCrc32(&Key.BlendState, sizeof(FBlendStateInitializerRHI));
-			CRC ^= FCrc::MemCrc32(&Key.RasterizerState, sizeof(FShaderRasterizerState));
-			CRC ^= FCrc::MemCrc32(&Key.DepthStencilState, sizeof(FDepthStencilStateInitializerRHI));
-			CRC ^= FCrc::MemCrc32(&Key.RenderTargets, sizeof(Key.RenderTargets));
-			CRC ^= FCrc::MemCrc32(&Key.SamplerStates, sizeof(Key.SamplerStates));
-			CRC ^= FCrc::MemCrc32(&Key.Resources, sizeof(Key.Resources));
-			CRC ^= GetTypeHash(Key.DepthStencilTarget);
-			CRC ^= GetTypeHash(Key.IndexType);
-			return CRC;
+			if(!Key.Hash)
+			{
+				Key.Hash = FCrc::MemCrc32(&Key.BlendState, sizeof(FBlendStateInitializerRHI));
+				Key.Hash ^= FCrc::MemCrc32(&Key.RasterizerState, sizeof(FShaderRasterizerState));
+				Key.Hash ^= FCrc::MemCrc32(&Key.DepthStencilState, sizeof(FDepthStencilStateInitializerRHI));
+				Key.Hash ^= FCrc::MemCrc32(&Key.RenderTargets, sizeof(Key.RenderTargets));
+				Key.Hash ^= FCrc::MemCrc32(&Key.SamplerStates, sizeof(Key.SamplerStates));
+				Key.Hash ^= FCrc::MemCrc32(&Key.Resources, sizeof(Key.Resources));
+				Key.Hash ^= GetTypeHash(Key.DepthStencilTarget);
+				Key.Hash ^= GetTypeHash(Key.IndexType);
+			}
+			return Key.Hash;
 		}
 		
 		friend FArchive& operator<<( FArchive& Ar, FShaderDrawKey& Info )
@@ -292,7 +318,7 @@ class SHADERCORE_API FShaderCache
 			{
 				Ar << Info.RenderTargets[i];
 			}
-			return Ar << Info.BlendState << Info.RasterizerState << Info.DepthStencilState << Info.DepthStencilTarget << Info.IndexType;
+			return Ar << Info.BlendState << Info.RasterizerState << Info.DepthStencilState << Info.DepthStencilTarget << Info.IndexType << Info.Hash;
 		}
 	};
 		
@@ -526,6 +552,9 @@ private:
 		FShaderCacheKey GeometryShader;
 		FShaderCacheKey HullShader;
 		FShaderCacheKey DomainShader;
+		mutable uint32 Hash;
+		
+		FShaderCacheBoundState() : Hash(0) {}
 		
 		friend bool operator ==(const FShaderCacheBoundState& A,const FShaderCacheBoundState& B)
 		{
@@ -545,19 +574,21 @@ private:
 
 		friend uint32 GetTypeHash(const FShaderCacheBoundState &Key)
 		{
-			uint32 Hash = 0;
-			
-			for(auto Element : Key.VertexDeclaration)
+			if(!Key.Hash)
 			{
-				Hash ^= FCrc::MemCrc_DEPRECATED(&Element, sizeof(FVertexElement));
+				for(auto Element : Key.VertexDeclaration)
+				{
+					Key.Hash ^= FCrc::MemCrc_DEPRECATED(&Element, sizeof(FVertexElement));
+				}
+				
+				Key.Hash ^= GetTypeHash(Key.VertexShader) ^ GetTypeHash(Key.PixelShader) ^ GetTypeHash(Key.GeometryShader) ^ GetTypeHash(Key.HullShader) ^ GetTypeHash(Key.DomainShader);
 			}
-			
-			return Hash ^ GetTypeHash(Key.VertexShader) ^ GetTypeHash(Key.PixelShader) ^ GetTypeHash(Key.GeometryShader) ^ GetTypeHash(Key.HullShader) ^ GetTypeHash(Key.DomainShader);
+			return Key.Hash;
 		}
 		
 		friend FArchive& operator<<( FArchive& Ar, FShaderCacheBoundState& Info )
 		{
-			return Ar << Info.VertexDeclaration << Info.VertexShader << Info.PixelShader << Info.GeometryShader << Info.HullShader << Info.DomainShader;
+			return Ar << Info.VertexDeclaration << Info.VertexShader << Info.PixelShader << Info.GeometryShader << Info.HullShader << Info.DomainShader << Info.Hash;
 		}
 	};
 	
