Further optimisations for ShaderCache.
- Several GetTypeHash would produce only very narrow hash distributions, so these have been modified.
- FShaderDrawKey::GetTypeHash was particularly slow (around 4% of sampled time in Mac UT in Instruments) almost all in MemCrc32, so this has been replaced with a much simpler hash routine to bring it down to ~1% sampled time.
- FShaderDrawKey::operator== comparison order changed to try and fail earlier by testing from most-to-least varying attributes.
- Changed the way sampler-state & texture look-up works to avoid unnecessary map lookup.

diff --git a/Engine/Source/Runtime/ShaderCore/Private/ShaderCache.cpp b/Engine/Source/Runtime/ShaderCore/Private/ShaderCache.cpp
--- a/Engine/Source/Runtime/ShaderCore/Private/ShaderCache.cpp
+++ b/Engine/Source/Runtime/ShaderCore/Private/ShaderCache.cpp
@@ -472,10 +472,9 @@ void FShaderCache::InternalLogSamplerState(FSamplerStateInitializerRHI const& In
 {
 	if ( bUseShaderDrawLog )
 	{
-		SamplerStates.Add(State, Init);
-		
 		FShaderPlatformCache& PlatformCache = Caches.FindOrAdd(GMaxRHIShaderPlatform);
-		PlatformCache.SamplerStates.Add(Init);
+		FSetElementId ID = PlatformCache.SamplerStates.Add(Init);
+		SamplerStates.Add(State, ID.AsInteger());
 	}
 }
 
@@ -483,14 +482,14 @@ void FShaderCache::InternalLogTexture(FShaderTextureKey const& Init, FTextureRHI
 {
 	if ( bUseShaderPredraw || bUseShaderDrawLog )
 	{
-		Textures.Add(State, Init);
-		CachedTextures.Add(Init, State);
-		
 		FShaderPlatformCache& PlatformCache = Caches.FindOrAdd(GMaxRHIShaderPlatform);
 		FShaderResourceKey Key;
 		Key.Tex = Init;
 		Key.Format = Init.Format;
-		PlatformCache.Resources.Add(Key);
+		FSetElementId ID = PlatformCache.Resources.Add(Key);
+		
+		Textures.Add(State, ID.AsInteger());
+		CachedTextures.Add(Init, State);
 	}
 }
 
@@ -498,8 +497,13 @@ void FShaderCache::InternalLogSRV(FShaderResourceViewRHIParamRef SRV, FTextureRH
 {
 	if ( bUseShaderPredraw || bUseShaderDrawLog )
 	{
+		FShaderPlatformCache& PlatformCache = Caches.FindOrAdd(GMaxRHIShaderPlatform);
+		
+		FSetElementId ID = FSetElementId::FromInteger(Textures.FindChecked(Texture));
+		FShaderResourceKey TexKey = PlatformCache.Resources[ID];
+		
 		FShaderResourceKey Key;
-		Key.Tex = Textures.FindChecked(Texture);
+		Key.Tex = TexKey.Tex;
 		Key.BaseMip = StartMip;
 		Key.MipLevels = NumMips;
 		Key.Format = Format;
@@ -507,7 +511,6 @@ void FShaderCache::InternalLogSRV(FShaderResourceViewRHIParamRef SRV, FTextureRH
 		SRVs.Add(SRV, Key);
 		CachedSRVs.Add(Key, FShaderResourceViewBinding(SRV, nullptr, Texture));
 		
-		FShaderPlatformCache& PlatformCache = Caches.FindOrAdd(GMaxRHIShaderPlatform);
 		PlatformCache.Resources.Add(Key);
 	}
 }
@@ -545,8 +548,12 @@ void FShaderCache::InternalRemoveTexture(FTextureRHIParamRef Texture)
 {
 	if ( bUseShaderPredraw || bUseShaderDrawLog )
 	{
-		auto Key = Textures.FindRef(Texture);
-		CachedTextures.Remove(Key);
+		FShaderPlatformCache& PlatformCache = Caches.FindOrAdd(GMaxRHIShaderPlatform);
+		
+		FSetElementId ID = FSetElementId::FromInteger(Textures.FindChecked(Texture));
+		FShaderResourceKey TexKey = PlatformCache.Resources[ID];
+		
+		CachedTextures.Remove(TexKey.Tex);
 		Textures.Remove(Texture);
 	}
 }
@@ -600,7 +607,9 @@ void FShaderCache::InternalSetRenderTargets( uint32 NumSimultaneousRenderTargets
 			if ( Target.Texture )
 			{
 				FShaderRenderTargetKey Key;
-				Key.Texture = Textures.FindChecked(Target.Texture);
+				FSetElementId ID = FSetElementId::FromInteger(Textures.FindChecked(Target.Texture));
+				FShaderResourceKey TexKey = PlatformCache.Resources[ID];
+				Key.Texture = TexKey.Tex;
 				check(Key.Texture.MipLevels == Target.Texture->GetNumMips());
 				Key.MipLevel = Key.Texture.MipLevels > Target.MipIndex ? Target.MipIndex : 0;
 				Key.ArrayIndex = Target.ArraySliceIndex;
@@ -615,7 +624,9 @@ void FShaderCache::InternalSetRenderTargets( uint32 NumSimultaneousRenderTargets
 		if ( NewDepthStencilTargetRHI && NewDepthStencilTargetRHI->Texture )
 		{
 			FShaderRenderTargetKey Key;
-			Key.Texture = Textures.FindChecked(NewDepthStencilTargetRHI->Texture);
+			FSetElementId ID = FSetElementId::FromInteger(Textures.FindChecked(NewDepthStencilTargetRHI->Texture));
+			FShaderResourceKey TexKey = PlatformCache.Resources[ID];
+			Key.Texture = TexKey.Tex;
 			CurrentDrawKey.DepthStencilTarget = PlatformCache.RenderTargets.Add(Key).AsInteger();
 		}
 		else
@@ -633,10 +644,7 @@ void FShaderCache::InternalSetSamplerState(EShaderFrequency Frequency, uint32 In
 		check(Index < GetFeatureLevelMaxTextureSamplers(GMaxRHIFeatureLevel));
 		if ( State )
 		{
-			FSamplerStateInitializerRHI Key = SamplerStates.FindChecked(State);
-			
-			FShaderPlatformCache& PlatformCache = Caches.FindOrAdd(GMaxRHIShaderPlatform);
-			CurrentDrawKey.SamplerStates[Frequency][Index] = PlatformCache.SamplerStates.FindId(Key).AsInteger();
+			CurrentDrawKey.SamplerStates[Frequency][Index] = SamplerStates.FindChecked(State);
 		}
 		else
 		{
@@ -661,11 +669,9 @@ void FShaderCache::InternalSetTexture(EShaderFrequency Frequency, uint32 Index,
 				Tex = State->GetTextureReference()->GetReferencedTexture();
 			}
 			
-			Key.Tex = Textures.FindChecked(Tex);
-			Key.Format = Key.Tex.Format;
-			
 			FShaderPlatformCache& PlatformCache = Caches.FindOrAdd(GMaxRHIShaderPlatform);
-			CurrentDrawKey.Resources[Frequency][Index] = PlatformCache.Resources.FindId(Key).AsInteger();
+			FSetElementId ID = FSetElementId::FromInteger(Textures.FindChecked(Tex));
+			CurrentDrawKey.Resources[Frequency][Index] = ID.AsInteger();
 		}
 		else
 		{
@@ -729,12 +735,20 @@ void FShaderCache::InternalLogDraw(uint8 IndexType)
 	{
 		FShaderPlatformCache& PlatformCache = Caches.FindOrAdd(GMaxRHIShaderPlatform);
 		CurrentDrawKey.IndexType = IndexType;
-		FSetElementId Id = PlatformCache.DrawStates.Add(CurrentDrawKey);
+		FSetElementId Id = PlatformCache.DrawStates.FindId(CurrentDrawKey);
+		if ( !Id.IsValidId() )
+		{
+			Id = PlatformCache.DrawStates.Add(CurrentDrawKey);
+		}
+		
 		TSet<int32>& ShaderDrawSet = PlatformCache.ShaderDrawStates.FindOrAdd(BoundShaderState);
-		ShaderDrawSet.Add(Id.AsInteger());
+		if( !ShaderDrawSet.Contains(Id.AsInteger()) )
+		{
+			ShaderDrawSet.Add(Id.AsInteger());
+		}
 		
 		// No need to predraw this shader draw key - we've already done it
-		ShadersToDraw.FindRef(BoundShaderState).Remove(PlatformCache.DrawStates.FindId(CurrentDrawKey).AsInteger());
+		ShadersToDraw.FindRef(BoundShaderState).Remove(Id.AsInteger());
 	}
 }
 
diff --git a/Engine/Source/Runtime/ShaderCore/Public/ShaderCache.h b/Engine/Source/Runtime/ShaderCore/Public/ShaderCache.h
--- a/Engine/Source/Runtime/ShaderCore/Public/ShaderCache.h
+++ b/Engine/Source/Runtime/ShaderCore/Public/ShaderCache.h
@@ -15,7 +15,7 @@ struct SHADERCORE_API FShaderCacheCustomVersion
 {
 	static const FGuid Key;
 	static const FGuid GameKey;
-	enum Type {	Initial, PreDraw, CacheHashes, Latest = CacheHashes };
+	enum Type {	Initial, PreDraw, CacheHashes, OptimisedHashes, Latest = OptimisedHashes };
 };
 
 enum EShaderCacheTextureType
@@ -42,7 +42,6 @@ struct SHADERCORE_API FShaderTextureKey
 	uint32 Samples;
 	uint8 Format;
 	TEnumAsByte<EShaderCacheTextureType> Type;
-
 	
 	FShaderTextureKey()
 	: Hash(0)
@@ -66,18 +65,18 @@ struct SHADERCORE_API FShaderTextureKey
 	{
 		if(!Key.Hash)
 		{
-			Key.Hash = Key.X;
-			Key.Hash ^= Key.Y;
+			Key.Hash = Key.X * 3;
+			Key.Hash ^= Key.Y * 2;
 			Key.Hash ^= Key.Z;
 			Key.Hash ^= Key.Flags;
-			Key.Hash ^= Key.MipLevels;
-			Key.Hash ^= Key.Samples;
-			Key.Hash ^= Key.Format;
+			Key.Hash ^= (Key.Format << 24);
+			Key.Hash ^= (Key.MipLevels << 16);
+			Key.Hash ^= (Key.Samples << 8);
 			Key.Hash ^= Key.Type;
 		}
 		return Key.Hash;
 	}
-		
+	
 	friend FArchive& operator<<( FArchive& Ar, FShaderTextureKey& Info )
 	{
 		return Ar << Info.Format << Info.Type << Info.Samples << Info.MipLevels << Info.Flags << Info.X << Info.Y << Info.Z << Info.Hash;
@@ -112,9 +111,9 @@ struct SHADERCORE_API FShaderResourceKey
 		if(!Key.Hash)
 		{
 			Key.Hash = GetTypeHash(Key.Tex);
-			Key.Hash ^= Key.BaseMip;
-			Key.Hash ^= Key.MipLevels;
-			Key.Hash ^= Key.Format;
+			Key.Hash ^= (Key.BaseMip << 24);
+			Key.Hash ^= (Key.MipLevels << 16);
+			Key.Hash ^= (Key.Format << 8);
 			Key.Hash ^= (uint32)Key.bSRV;
 		}
 		return Key.Hash;
@@ -151,7 +150,7 @@ struct SHADERCORE_API FShaderRenderTargetKey
 		if(!Key.Hash)
 		{
 			Key.Hash = GetTypeHash(Key.Texture);
-			Key.Hash ^= Key.MipLevel;
+			Key.Hash ^= (Key.MipLevel << 8);
 			Key.Hash ^= Key.ArrayIndex;
 		}
 		return Key.Hash;
@@ -168,7 +167,7 @@ class SHADERCORE_API FShaderCache
 	struct SHADERCORE_API FShaderCacheKey
 	{
 		FShaderCacheKey() : Platform(SP_NumPlatforms), Frequency(SF_NumFrequencies), Hash(0), bActive(false) {}
-	
+		
 		FSHAHash SHAHash;
 		EShaderPlatform Platform;
 		EShaderFrequency Frequency;
@@ -179,14 +178,14 @@ class SHADERCORE_API FShaderCache
 		{
 			return A.SHAHash == B.SHAHash && A.Platform == B.Platform && A.Frequency == B.Frequency && A.bActive == B.bActive;
 		}
-
+		
 		friend uint32 GetTypeHash(const FShaderCacheKey &Key)
 		{
 			if(!Key.Hash)
 			{
 				uint32 TargetFrequency = Key.Frequency;
 				uint32 TargetPlatform = Key.Platform;
-				Key.Hash = FCrc::MemCrc_DEPRECATED((const void*)&Key.SHAHash, sizeof(Key.SHAHash)) ^ GetTypeHash(TargetPlatform) ^ GetTypeHash(TargetFrequency) ^ GetTypeHash(Key.bActive);
+				Key.Hash = FCrc::MemCrc_DEPRECATED((const void*)&Key.SHAHash, sizeof(Key.SHAHash)) ^ GetTypeHash(TargetPlatform) ^ (GetTypeHash(TargetFrequency) << 16) ^ GetTypeHash(Key.bActive);
 			}
 			return Key.Hash;
 		}
@@ -263,8 +262,18 @@ class SHADERCORE_API FShaderCache
 				Ar << RasterizerStateInitializer.bEnableLineAA;
 				return Ar;
 			}
+			
+			friend uint32 GetTypeHash(const FShaderRasterizerState &Key)
+			{
+				uint32 Hash = (*((uint32*)&Key.DepthBias) ^ *((uint32*)&Key.SlopeScaleDepthBias));
+				Hash ^= (Key.FillMode << 8);
+				Hash ^= Key.CullMode;
+				Hash ^= Key.bAllowMSAA ? 2 : 0;
+				Hash ^= Key.bEnableLineAA ? 1 : 0;
+				return Hash;
+			}
 		};
-	
+		
 		FBlendStateInitializerRHI BlendState;
 		FShaderRasterizerState RasterizerState;
 		FDepthStencilStateInitializerRHI DepthStencilState;
@@ -278,13 +287,13 @@ class SHADERCORE_API FShaderCache
 		friend bool operator ==(const FShaderDrawKey& A,const FShaderDrawKey& B)
 		{
 			bool Compare = A.IndexType == B.IndexType &&
-							A.DepthStencilTarget == B.DepthStencilTarget &&
-							FMemory::Memcmp(&A.BlendState, &B.BlendState, sizeof(FBlendStateInitializerRHI)) == 0 &&
-							FMemory::Memcmp(&A.RasterizerState, &B.RasterizerState, sizeof(FShaderRasterizerState)) == 0 &&
-							FMemory::Memcmp(&A.DepthStencilState, &B.DepthStencilState, sizeof(FDepthStencilStateInitializerRHI)) == 0 &&
-							FMemory::Memcmp(&A.RenderTargets, &B.RenderTargets, sizeof(A.RenderTargets)) == 0 &&
-							FMemory::Memcmp(&A.SamplerStates, &B.SamplerStates, sizeof(A.SamplerStates)) == 0 &&
-							FMemory::Memcmp(&A.Resources, &B.Resources, sizeof(A.Resources)) == 0;
+			A.DepthStencilTarget == B.DepthStencilTarget &&
+			FMemory::Memcmp(&A.Resources, &B.Resources, sizeof(A.Resources)) == 0 &&
+			FMemory::Memcmp(&A.SamplerStates, &B.SamplerStates, sizeof(A.SamplerStates)) == 0 &&
+			FMemory::Memcmp(&A.BlendState, &B.BlendState, sizeof(FBlendStateInitializerRHI)) == 0 &&
+			FMemory::Memcmp(&A.RasterizerState, &B.RasterizerState, sizeof(FShaderRasterizerState)) == 0 &&
+			FMemory::Memcmp(&A.DepthStencilState, &B.DepthStencilState, sizeof(FDepthStencilStateInitializerRHI)) == 0 &&
+			FMemory::Memcmp(&A.RenderTargets, &B.RenderTargets, sizeof(A.RenderTargets)) == 0;
 			return Compare;
 		}
 		
@@ -292,14 +301,47 @@ class SHADERCORE_API FShaderCache
 		{
 			if(!Key.Hash)
 			{
-				Key.Hash = FCrc::MemCrc32(&Key.BlendState, sizeof(FBlendStateInitializerRHI));
-				Key.Hash ^= FCrc::MemCrc32(&Key.RasterizerState, sizeof(FShaderRasterizerState));
-				Key.Hash ^= FCrc::MemCrc32(&Key.DepthStencilState, sizeof(FDepthStencilStateInitializerRHI));
-				Key.Hash ^= FCrc::MemCrc32(&Key.RenderTargets, sizeof(Key.RenderTargets));
-				Key.Hash ^= FCrc::MemCrc32(&Key.SamplerStates, sizeof(Key.SamplerStates));
-				Key.Hash ^= FCrc::MemCrc32(&Key.Resources, sizeof(Key.Resources));
+				Key.Hash ^= (Key.BlendState.bUseIndependentRenderTargetBlendStates ? (1 << 31) : 0);
+				for( uint32 i = 0; i < MaxSimultaneousRenderTargets; i++ )
+				{
+					Key.Hash ^= (Key.BlendState.RenderTargets[i].ColorBlendOp << 24);
+					Key.Hash ^= (Key.BlendState.RenderTargets[i].ColorSrcBlend << 16);
+					Key.Hash ^= (Key.BlendState.RenderTargets[i].ColorDestBlend << 8);
+					Key.Hash ^= (Key.BlendState.RenderTargets[i].ColorWriteMask << 0);
+					Key.Hash ^= (Key.BlendState.RenderTargets[i].AlphaBlendOp << 24);
+					Key.Hash ^= (Key.BlendState.RenderTargets[i].AlphaSrcBlend << 16);
+					Key.Hash ^= (Key.BlendState.RenderTargets[i].AlphaDestBlend << 8);
+					Key.Hash ^= Key.RenderTargets[i];
+				}
+				
+				for( uint32 i = 0; i < SF_NumFrequencies; i++ )
+				{
+					for( uint32 j = 0; j < MaxTextureSamplers; j++ )
+					{
+						Key.Hash ^= Key.SamplerStates[i][j];
+						Key.Hash ^= Key.Resources[i][j];
+					}
+				}
+				
+				Key.Hash ^= (Key.DepthStencilState.bEnableDepthWrite ? (1 << 31) : 0);
+				Key.Hash ^= (Key.DepthStencilState.DepthTest << 24);
+				Key.Hash ^= (Key.DepthStencilState.bEnableFrontFaceStencil ? (1 << 23) : 0);
+				Key.Hash ^= (Key.DepthStencilState.FrontFaceStencilTest << 24);
+				Key.Hash ^= (Key.DepthStencilState.FrontFaceStencilFailStencilOp << 16);
+				Key.Hash ^= (Key.DepthStencilState.FrontFaceDepthFailStencilOp << 8);
+				Key.Hash ^= (Key.DepthStencilState.FrontFacePassStencilOp);
+				Key.Hash ^= (Key.DepthStencilState.bEnableBackFaceStencil ? (1 << 15) : 0);
+				Key.Hash ^= (Key.DepthStencilState.BackFaceStencilTest << 24);
+				Key.Hash ^= (Key.DepthStencilState.BackFaceStencilFailStencilOp << 16);
+				Key.Hash ^= (Key.DepthStencilState.BackFaceDepthFailStencilOp << 8);
+				Key.Hash ^= (Key.DepthStencilState.BackFacePassStencilOp);
+				Key.Hash ^= (Key.DepthStencilState.StencilReadMask << 8);
+				Key.Hash ^= (Key.DepthStencilState.StencilWriteMask);
+				
 				Key.Hash ^= GetTypeHash(Key.DepthStencilTarget);
 				Key.Hash ^= GetTypeHash(Key.IndexType);
+				
+				Key.Hash ^= GetTypeHash(Key.RasterizerState);
 			}
 			return Key.Hash;
 		}
@@ -321,15 +363,15 @@ class SHADERCORE_API FShaderCache
 			return Ar << Info.BlendState << Info.RasterizerState << Info.DepthStencilState << Info.DepthStencilTarget << Info.IndexType << Info.Hash;
 		}
 	};
-		
+	
 public:
 	FShaderCache();
 	~FShaderCache();
-
+	
 	// Called by the game to set the game specific shader cache version, only caches of this version will be loaded.
 	// Must be called before RHI initialisation, as InitShaderCache will load any existing cache.
 	static void SetGameVersion(int32 InGameVersion);
-
+	
 	// Called by the RHI implementation, not to be called by the game.
 	static void InitShaderCache();
 	static void ShutdownShaderCache();
@@ -338,7 +380,7 @@ public:
 	{
 		return bUseShaderCaching ? Cache : nullptr;
 	}
-
+	
 	FVertexShaderRHIRef GetVertexShader(EShaderPlatform Platform, FSHAHash Hash, TArray<uint8> const& Code);
 	FPixelShaderRHIRef GetPixelShader(EShaderPlatform Platform, FSHAHash Hash, TArray<uint8> const& Code);
 	FGeometryShaderRHIRef GetGeometryShader(EShaderPlatform Platform, FSHAHash Hash, TArray<uint8> const& Code);
@@ -540,7 +582,7 @@ public:
 			Cache->InternalPreDrawShaders(RHICmdList);
 		}
 	}
-		
+	
 	friend FArchive& operator<<( FArchive& Ar, FShaderCache& Info );
 	
 private:
@@ -571,7 +613,7 @@ private:
 			}
 			return false;
 		}
-
+		
 		friend uint32 GetTypeHash(const FShaderCacheBoundState &Key)
 		{
 			if(!Key.Hash)
@@ -613,7 +655,7 @@ private:
 		/** Calculates a hash index for a key. */
 		static uint32 GetKeyHash(KeyInitType Key);
 	};
-
+	
 	struct FShaderPlatformCache
 	{
 		TSet<FShaderCacheKey> Shaders;
@@ -630,9 +672,9 @@ private:
 			return Ar << Info.Shaders << Info.BoundShaderStates << Info.DrawStates << Info.RenderTargets << Info.Resources << Info.SamplerStates << Info.ShaderStateMembership << Info.ShaderDrawStates;
 		}
 	};
-
+	
 	typedef TMap<uint32, FShaderPlatformCache> FShaderCaches;
-
+	
 	struct FShaderResourceViewBinding
 	{
 		FShaderResourceViewBinding()
@@ -651,7 +693,7 @@ private:
 		FVertexBufferRHIParamRef VertexBuffer;
 		FTextureRHIParamRef Texture;
 	};
-
+	
 	struct FShaderTextureBinding
 	{
 		FShaderTextureBinding() {}
@@ -685,19 +727,19 @@ private:
 		FVertexBufferRHIRef VertexBuffer;
 		FTextureRHIRef Texture;
 	};
-
+	
 private:
 	void InternalLogShader(EShaderPlatform Platform, EShaderFrequency Frequency, FSHAHash Hash, TArray<uint8> const& Code);
 	
 	void InternalLogVertexDeclaration(const FVertexDeclarationElementList& VertexElements, FVertexDeclarationRHIParamRef VertexDeclaration);
 	
 	void InternalLogBoundShaderState(EShaderPlatform Platform, FVertexDeclarationRHIParamRef VertexDeclaration,
-								FVertexShaderRHIParamRef VertexShader,
-								FPixelShaderRHIParamRef PixelShader,
-								FHullShaderRHIParamRef HullShader,
-								FDomainShaderRHIParamRef DomainShader,
-								FGeometryShaderRHIParamRef GeometryShader,
-								FBoundShaderStateRHIParamRef BoundState);
+									 FVertexShaderRHIParamRef VertexShader,
+									 FPixelShaderRHIParamRef PixelShader,
+									 FHullShaderRHIParamRef HullShader,
+									 FDomainShaderRHIParamRef DomainShader,
+									 FGeometryShaderRHIParamRef GeometryShader,
+									 FBoundShaderStateRHIParamRef BoundState);
 	
 	void InternalLogBlendState(FBlendStateInitializerRHI const& Init, FBlendStateRHIParamRef State);
 	void InternalLogRasterizerState(FRasterizerStateInitializerRHI const& Init, FRasterizerStateRHIParamRef State);
@@ -751,8 +793,8 @@ private:
 	TMap<FBlendStateRHIParamRef, FBlendStateInitializerRHI> BlendStates;
 	TMap<FRasterizerStateRHIParamRef, FRasterizerStateInitializerRHI> RasterizerStates;
 	TMap<FDepthStencilStateRHIParamRef, FDepthStencilStateInitializerRHI> DepthStencilStates;
-	TMap<FSamplerStateRHIParamRef, FSamplerStateInitializerRHI> SamplerStates;
-	TMap<FTextureRHIParamRef, FShaderTextureKey> Textures;
+	TMap<FSamplerStateRHIParamRef, int32> SamplerStates;
+	TMap<FTextureRHIParamRef, int32> Textures;
 	TMap<FShaderResourceViewRHIParamRef, FShaderResourceKey> SRVs;
 	
 	// Shaders we need to predraw
