Fix non-unity builds.

diff --git a/Engine/Source/Runtime/ShaderCore/Private/ShaderCache.cpp b/Engine/Source/Runtime/ShaderCore/Private/ShaderCache.cpp
--- a/Engine/Source/Runtime/ShaderCore/Private/ShaderCache.cpp
+++ b/Engine/Source/Runtime/ShaderCore/Private/ShaderCache.cpp
@@ -6,6 +6,7 @@
 
 #include "ShaderCore.h"
 #include "ShaderCache.h"
+#include "Shader.h"
 #include "RHI.h"
 #include "RenderingThread.h"
 
