In the shader-cache associate logged draw states with a 'streaming key' so that they may be batched and predrawn only when that 'streaming key' is active. This limits predrawing to only those states that are currently relevant. Added comments as to the purpose and operation of the shader-cache for posterity.

diff --git a/Engine/Source/Runtime/ShaderCore/Private/ShaderCache.cpp b/Engine/Source/Runtime/ShaderCore/Private/ShaderCache.cpp
--- a/Engine/Source/Runtime/ShaderCore/Private/ShaderCache.cpp
+++ b/Engine/Source/Runtime/ShaderCore/Private/ShaderCache.cpp
@@ -128,7 +128,8 @@ FArchive& operator<<( FArchive& Ar, FShaderCache& Info )
 }
 
 FShaderCache::FShaderCache()
-: bCurrentDepthStencilTarget(false)
+: StreamingKey(0)
+, bCurrentDepthStencilTarget(false)
 , CurrentNumRenderTargets(0)
 , CurrentShaderState(nullptr)
 , bIsPreDraw(false)
@@ -307,6 +308,39 @@ FComputeShaderRHIRef FShaderCache::GetComputeShader(EShaderPlatform Platform, TA
 	return Shader;
 }
 
+void FShaderCache::InternalLogStreamingKey(uint32 StreamKey, bool const bActive)
+{
+	// Defer to the render thread to avoid race conditions
+	ENQUEUE_UNIQUE_RENDER_COMMAND_THREEPARAMETER(
+	FShaderCacheInternalLogLevel,
+	FShaderCache*,Cache,Cache,
+	uint32,StreamKey,StreamKey,
+	bool,bActive,bActive,
+	{
+		if(bActive)
+		{
+			Cache->ActiveStreamingKeys.Add(StreamKey);
+		}
+		else
+		{
+			Cache->ActiveStreamingKeys.Remove(StreamKey);
+		}
+		
+		uint32 NewStreamingKey = 0;
+		for(uint32 Key : Cache->ActiveStreamingKeys)
+		{
+			NewStreamingKey ^= Key;
+		}
+		Cache->StreamingKey = NewStreamingKey;
+		
+		if(!Cache->ShadersToDraw.Contains(NewStreamingKey))
+		{
+			FShaderPlatformCache& PlatformCache = Cache->Caches.FindOrAdd(GMaxRHIShaderPlatform);
+			Cache->ShadersToDraw.Add(NewStreamingKey, PlatformCache.StreamingDrawStates.FindRef(NewStreamingKey));
+		}
+	});
+}
+
 void FShaderCache::InternalLogShader(EShaderPlatform Platform, EShaderFrequency Frequency, FSHAHash Hash, TArray<uint8> const& Code)
 {
 	bool bUsable = FShaderResource::ArePlatformsCompatible(GMaxRHIShaderPlatform, Platform);
@@ -742,20 +776,23 @@ void FShaderCache::InternalLogDraw(uint8 IndexType)
 			Id = PlatformCache.DrawStates.Add(CurrentDrawKey);
 		}
 		
-		TSet<int32>& ShaderDrawSet = PlatformCache.ShaderDrawStates.FindOrAdd(BoundShaderState);
+		TSet<int32>& ShaderDrawSet = PlatformCache.StreamingDrawStates.FindOrAdd(StreamingKey).ShaderDrawStates.FindOrAdd(BoundShaderState);
 		if( !ShaderDrawSet.Contains(Id.AsInteger()) )
 		{
 			ShaderDrawSet.Add(Id.AsInteger());
 		}
 		
 		// No need to predraw this shader draw key - we've already done it
-		ShadersToDraw.FindRef(BoundShaderState).Remove(Id.AsInteger());
+		for(auto StreamingMap : ShadersToDraw)
+		{
+			StreamingMap.Value.ShaderDrawStates.FindRef(BoundShaderState).Remove(Id.AsInteger());
+		}
 	}
 }
 
 void FShaderCache::InternalPreDrawShaders(FRHICommandList& RHICmdList)
 {
-	if ( bUseShaderPredraw && ShadersToDraw.Num() > 0 )
+	if ( bUseShaderPredraw && ShadersToDraw.FindRef(StreamingKey).ShaderDrawStates.Num() > 0 )
 	{
 		bIsPreDraw = true;
 		
@@ -789,7 +826,8 @@ void FShaderCache::InternalPreDrawShaders(FRHICommandList& RHICmdList)
 		RHICmdList.SetViewport(0, 0, FLT_MIN, 3, 3, FLT_MAX);
 		
 		int64 TimeForPredrawing = 0;
-		for ( auto It = ShadersToDraw.CreateIterator(); (PredrawBatchTime == -1 || TimeForPredrawing < PredrawBatchTime) && It; ++It )
+		TMap<FShaderCacheBoundState, TSet<int32>>& ShaderDrawStates = ShadersToDraw.FindOrAdd(StreamingKey).ShaderDrawStates;
+		for ( auto It = ShaderDrawStates.CreateIterator(); (PredrawBatchTime == -1 || TimeForPredrawing < PredrawBatchTime) && It; ++It )
 		{
 			uint32 Start = FPlatformTime::Cycles();
 			
@@ -818,7 +856,7 @@ void FShaderCache::InternalPreDrawShaders(FRHICommandList& RHICmdList)
 
 		RHICmdList.SetViewport(Viewport[0], Viewport[1], DepthRange[0], Viewport[2], Viewport[3], DepthRange[1]);
 		
-		if ( ShadersToDraw.Num() == 0 )
+		if ( ShadersToDraw.FindOrAdd(StreamingKey).ShaderDrawStates.Num() == 0 )
 		{
 			PredrawRTs.Empty();
 			PredrawBindings.Empty();
@@ -871,13 +909,21 @@ void FShaderCache::PrebindShader(FShaderCacheKey const& Key)
 							BoundShaderStates.Add(State, BoundState);
 							if (bUseShaderPredraw)
 							{
-								ShadersToDraw.Add(State, PlatformCache.ShaderDrawStates.FindOrAdd(State));
+								TSet<int32>& StreamCache = PlatformCache.StreamingDrawStates.FindOrAdd(StreamingKey).ShaderDrawStates.FindOrAdd(State);
+								if(!ShadersToDraw.FindOrAdd(StreamingKey).ShaderDrawStates.Contains(State))
+								{
+									ShadersToDraw.FindOrAdd(StreamingKey).ShaderDrawStates.Add(State, StreamCache);
+								}
 							}
 						}
 					}
 					else if (bUseShaderPredraw)
 					{
-						ShadersToDraw.Add(State, PlatformCache.ShaderDrawStates.FindOrAdd(State));
+						TSet<int32>& StreamCache = PlatformCache.StreamingDrawStates.FindOrAdd(StreamingKey).ShaderDrawStates.FindOrAdd(State);
+						if(!ShadersToDraw.FindOrAdd(StreamingKey).ShaderDrawStates.Contains(State))
+						{
+							ShadersToDraw.FindOrAdd(StreamingKey).ShaderDrawStates.Add(State, StreamCache);
+						}
 					}
 				}
 			}
diff --git a/Engine/Source/Runtime/ShaderCore/Public/ShaderCache.h b/Engine/Source/Runtime/ShaderCore/Public/ShaderCache.h
--- a/Engine/Source/Runtime/ShaderCore/Public/ShaderCache.h
+++ b/Engine/Source/Runtime/ShaderCore/Public/ShaderCache.h
@@ -15,9 +15,10 @@ struct SHADERCORE_API FShaderCacheCustomVersion
 {
 	static const FGuid Key;
 	static const FGuid GameKey;
-	enum Type {	Initial, PreDraw, CacheHashes, OptimisedHashes, Latest = OptimisedHashes };
+	enum Type {	Initial, PreDraw, CacheHashes, OptimisedHashes, StreamingKeys, Latest = StreamingKeys };
 };
 
+/** Texture type enum for shader cache draw keys */
 enum EShaderCacheTextureType
 {
 	SCTT_Invalid,
@@ -31,6 +32,7 @@ enum EShaderCacheTextureType
 	SCTT_Buffer
 };
 
+/** The minimum texture state required for logging shader draw states */
 struct SHADERCORE_API FShaderTextureKey
 {
 	mutable uint32 Hash;
@@ -83,6 +85,7 @@ struct SHADERCORE_API FShaderTextureKey
 	}
 };
 
+/** SRV state tracked by the shader-cache to properly predraw shaders */
 struct SHADERCORE_API FShaderResourceKey
 {
 	FShaderTextureKey Tex;
@@ -125,6 +128,7 @@ struct SHADERCORE_API FShaderResourceKey
 	}
 };
 
+/** Render target state tracked for predraw */
 struct SHADERCORE_API FShaderRenderTargetKey
 {
 	FShaderTextureKey Texture;
@@ -162,6 +166,36 @@ struct SHADERCORE_API FShaderRenderTargetKey
 	}
 };
 
+/**
+ * The FShaderCache provides mechanisms for reducing shader hitching in-game:
+ *	- Early submission during shader deserialisation rather than on-demand (r.UseShaderCaching 0/1).
+ *	- Tracking of bound-shader-states so that they may be pre-bound during early submission.
+ *	- Tracking of RHI draw states so that each bound-shader-state (r.UseShaderDrawLog 0/1).
+ *	- Predrawing of tracked RHI draw states to eliminate first-use hitches (r.UseShaderPredraw 0/1).
+ *	- Control over time spent predrawing each frame (r.PredrawBatchTime Time (ms), -1 for all).
+ *
+ * The cache should be populated by enabling r.UseShaderCaching & r.UseShaderDrawLog on a development machine. 
+ * Users/players should then consume the cache by enabling r.UseShaderCaching & r.UseShaderPredraw. 
+ * Draw logging (r.UseShaderDrawLog) adds noticeable fixed overhead so should be avoided if possible. 
+ * Since the caching is done via shader hashes it is advisable to only use this as a final optimisation tool 
+ * when content is largely complete as changes to shader hashes will result in unusued entries accumulating in the cache, 
+ * increasing cache size without reducing hitches.
+ * 
+ * Cache locations:
+ *  - While populating: <Game>/Saved/ShaderCache.ushadercache
+ *  - For distribution: <Game>/Content/ShaderCache.ushadercache
+ * The code will first try and load the writeable cache, then fallback to the distribution cache if needed.
+ * 
+ * Handling Updates/Invalidation:
+ * When the cache needs to be updated & writable caches invalidated the game should specify a new GameVersion.
+ * Call FShaderCache::SetGameVersion before initialisating the RHI (which initialises the cache), which will cause 
+ * the contents of caches from previous version to be ignored. Presently you cannot carry over cache entries from a previous version.
+ *
+ * Region/Stream Batching:
+ * For streaming games, or where the cache becomes very large, calls to FShaderCache::SetStreamingKey should be added with unique values for
+ * the currently relevant game regions/streaming levels (as required). Logged draw states will be linked to the active streaming key.
+ * This limits predrawing to only those draw states required by the active streaming key on subsequent runs.
+ */
 class SHADERCORE_API FShaderCache
 {
 	struct SHADERCORE_API FShaderCacheKey
@@ -389,9 +423,16 @@ public:
 	FComputeShaderRHIRef GetComputeShader(EShaderPlatform Platform, TArray<uint8> const& Code);
 	FVertexDeclarationRHIParamRef GetVertexDeclaration(FVertexDeclarationElementList& VertexElements);
 	
+	static FORCEINLINE void LogStreamingKey(uint32 StreamingKey, bool const bActive)
+	{
+		if ( Cache )
+		{
+			Cache->InternalLogStreamingKey(StreamingKey, bActive);
+		}
+	}
+	
 	static FORCEINLINE void LogShader(EShaderPlatform Platform, EShaderFrequency Frequency, FSHAHash Hash, TArray<uint8> const& Code)
 	{
-		FShaderCache* ShaderCache = GetShaderCache();
 		if ( Cache )
 		{
 			Cache->InternalLogShader(Platform, Frequency, Hash, Code);
@@ -400,7 +441,6 @@ public:
 	
 	static FORCEINLINE void LogVertexDeclaration(const FVertexDeclarationElementList& VertexElements, FVertexDeclarationRHIParamRef VertexDeclaration)
 	{
-		FShaderCache* ShaderCache = GetShaderCache();
 		if ( Cache )
 		{
 			Cache->InternalLogVertexDeclaration(VertexElements, VertexDeclaration);
@@ -415,7 +455,6 @@ public:
 												FGeometryShaderRHIParamRef GeometryShader,
 												FBoundShaderStateRHIParamRef BoundState)
 	{
-		FShaderCache* ShaderCache = GetShaderCache();
 		if ( Cache )
 		{
 			Cache->InternalLogBoundShaderState(Platform, VertexDeclaration, VertexShader, PixelShader, HullShader, DomainShader, GeometryShader, BoundState);
@@ -424,7 +463,6 @@ public:
 	
 	static FORCEINLINE void LogBlendState(FBlendStateInitializerRHI const& Init, FBlendStateRHIParamRef State)
 	{
-		FShaderCache* ShaderCache = GetShaderCache();
 		if ( Cache )
 		{
 			Cache->InternalLogBlendState(Init, State);
@@ -432,7 +470,6 @@ public:
 	}
 	static FORCEINLINE void LogRasterizerState(FRasterizerStateInitializerRHI const& Init, FRasterizerStateRHIParamRef State)
 	{
-		FShaderCache* ShaderCache = GetShaderCache();
 		if ( Cache )
 		{
 			Cache->InternalLogRasterizerState(Init, State);
@@ -440,7 +477,6 @@ public:
 	}
 	static FORCEINLINE void LogDepthStencilState(FDepthStencilStateInitializerRHI const& Init, FDepthStencilStateRHIParamRef State)
 	{
-		FShaderCache* ShaderCache = GetShaderCache();
 		if ( Cache )
 		{
 			Cache->InternalLogDepthStencilState(Init, State);
@@ -448,7 +484,6 @@ public:
 	}
 	static FORCEINLINE void LogSamplerState(FSamplerStateInitializerRHI const& Init, FSamplerStateRHIParamRef State)
 	{
-		FShaderCache* ShaderCache = GetShaderCache();
 		if ( Cache )
 		{
 			Cache->InternalLogSamplerState(Init, State);
@@ -456,7 +491,6 @@ public:
 	}
 	static FORCEINLINE void LogTexture(FShaderTextureKey const& Init, FTextureRHIParamRef State)
 	{
-		FShaderCache* ShaderCache = GetShaderCache();
 		if ( Cache )
 		{
 			Cache->InternalLogTexture(Init, State);
@@ -464,7 +498,6 @@ public:
 	}
 	static FORCEINLINE void LogSRV(FShaderResourceViewRHIParamRef SRV, FTextureRHIParamRef Texture, uint8 StartMip, uint8 NumMips, uint8 Format)
 	{
-		FShaderCache* ShaderCache = GetShaderCache();
 		if ( Cache )
 		{
 			Cache->InternalLogSRV(SRV, Texture, StartMip, NumMips, Format);
@@ -472,7 +505,6 @@ public:
 	}
 	static FORCEINLINE void LogSRV(FShaderResourceViewRHIParamRef SRV, FVertexBufferRHIParamRef Vb, uint32 Stride, uint8 Format)
 	{
-		FShaderCache* ShaderCache = GetShaderCache();
 		if ( Cache )
 		{
 			Cache->InternalLogSRV(SRV, Vb, Stride, Format);
@@ -480,7 +512,6 @@ public:
 	}
 	static FORCEINLINE void RemoveSRV(FShaderResourceViewRHIParamRef SRV)
 	{
-		FShaderCache* ShaderCache = GetShaderCache();
 		if ( Cache )
 		{
 			Cache->InternalRemoveSRV(SRV);
@@ -488,7 +519,6 @@ public:
 	}
 	static FORCEINLINE void RemoveTexture(FTextureRHIParamRef Texture)
 	{
-		FShaderCache* ShaderCache = GetShaderCache();
 		if ( Cache )
 		{
 			Cache->InternalRemoveTexture(Texture);
@@ -496,7 +526,6 @@ public:
 	}
 	static FORCEINLINE void SetBlendState(FBlendStateRHIParamRef State)
 	{
-		FShaderCache* ShaderCache = GetShaderCache();
 		if ( Cache )
 		{
 			Cache->InternalSetBlendState(State);
@@ -504,7 +533,6 @@ public:
 	}
 	static FORCEINLINE void SetRasterizerState(FRasterizerStateRHIParamRef State)
 	{
-		FShaderCache* ShaderCache = GetShaderCache();
 		if ( Cache )
 		{
 			Cache->InternalSetRasterizerState(State);
@@ -512,7 +540,6 @@ public:
 	}
 	static FORCEINLINE void SetDepthStencilState(FDepthStencilStateRHIParamRef State)
 	{
-		FShaderCache* ShaderCache = GetShaderCache();
 		if ( Cache )
 		{
 			Cache->InternalSetDepthStencilState(State);
@@ -520,7 +547,6 @@ public:
 	}
 	static FORCEINLINE void SetRenderTargets( uint32 NumSimultaneousRenderTargets, const FRHIRenderTargetView* NewRenderTargetsRHI, const FRHIDepthRenderTargetView* NewDepthStencilTargetRHI )
 	{
-		FShaderCache* ShaderCache = GetShaderCache();
 		if ( Cache )
 		{
 			Cache->InternalSetRenderTargets(NumSimultaneousRenderTargets, NewRenderTargetsRHI, NewDepthStencilTargetRHI);
@@ -528,7 +554,6 @@ public:
 	}
 	static FORCEINLINE void SetSamplerState(EShaderFrequency Frequency, uint32 Index, FSamplerStateRHIParamRef State)
 	{
-		FShaderCache* ShaderCache = GetShaderCache();
 		if ( Cache )
 		{
 			Cache->InternalSetSamplerState(Frequency, Index, State);
@@ -536,7 +561,6 @@ public:
 	}
 	static FORCEINLINE void SetTexture(EShaderFrequency Frequency, uint32 Index, FTextureRHIParamRef State)
 	{
-		FShaderCache* ShaderCache = GetShaderCache();
 		if ( Cache )
 		{
 			Cache->InternalSetTexture(Frequency, Index, State);
@@ -544,7 +568,6 @@ public:
 	}
 	static FORCEINLINE void SetSRV(EShaderFrequency Frequency, uint32 Index, FShaderResourceViewRHIParamRef SRV)
 	{
-		FShaderCache* ShaderCache = GetShaderCache();
 		if ( Cache )
 		{
 			Cache->InternalSetSRV(Frequency, Index, SRV);
@@ -552,7 +575,6 @@ public:
 	}
 	static FORCEINLINE void SetBoundShaderState(FBoundShaderStateRHIParamRef State)
 	{
-		FShaderCache* ShaderCache = GetShaderCache();
 		if ( Cache )
 		{
 			Cache->InternalSetBoundShaderState(State);
@@ -560,7 +582,6 @@ public:
 	}
 	static FORCEINLINE void SetViewport(uint32 MinX, uint32 MinY, float MinZ, uint32 MaxX, uint32 MaxY, float MaxZ)
 	{
-		FShaderCache* ShaderCache = GetShaderCache();
 		if ( Cache )
 		{
 			Cache->InternalSetViewport(MinX, MinY, MinZ, MaxX, MaxY, MaxZ);
@@ -568,7 +589,6 @@ public:
 	}
 	static FORCEINLINE void LogDraw(uint8 IndexType)
 	{
-		FShaderCache* ShaderCache = GetShaderCache();
 		if ( Cache )
 		{
 			Cache->InternalLogDraw(IndexType);
@@ -576,7 +596,6 @@ public:
 	}
 	static FORCEINLINE void PreDrawShaders(FRHICommandList& RHICmdList)
 	{
-		FShaderCache* ShaderCache = GetShaderCache();
 		if ( Cache )
 		{
 			Cache->InternalPreDrawShaders(RHICmdList);
@@ -656,6 +675,16 @@ private:
 		static uint32 GetKeyHash(KeyInitType Key);
 	};
 	
+	struct FShaderStreamingCache
+	{
+		TMap<FShaderCacheBoundState, TSet<int32>> ShaderDrawStates;
+		
+		friend FArchive& operator<<( FArchive& Ar, FShaderStreamingCache& Info )
+		{
+			return Ar << Info.ShaderDrawStates;
+		}
+	};
+	
 	struct FShaderPlatformCache
 	{
 		TSet<FShaderCacheKey> Shaders;
@@ -665,11 +694,11 @@ private:
 		TSet<FShaderResourceKey> Resources;
 		TSet<FSamplerStateInitializerRHI, FSamplerStateInitializerRHIKeyFuncs> SamplerStates;
 		TMap<FShaderCacheKey, TSet<FShaderCacheBoundState>> ShaderStateMembership;
-		TMap<FShaderCacheBoundState, TSet<int32>> ShaderDrawStates;
+		TMap<uint32, FShaderStreamingCache> StreamingDrawStates;
 		
 		friend FArchive& operator<<( FArchive& Ar, FShaderPlatformCache& Info )
 		{
-			return Ar << Info.Shaders << Info.BoundShaderStates << Info.DrawStates << Info.RenderTargets << Info.Resources << Info.SamplerStates << Info.ShaderStateMembership << Info.ShaderDrawStates;
+			return Ar << Info.Shaders << Info.BoundShaderStates << Info.DrawStates << Info.RenderTargets << Info.Resources << Info.SamplerStates << Info.ShaderStateMembership << Info.StreamingDrawStates;
 		}
 	};
 	
@@ -729,6 +758,7 @@ private:
 	};
 	
 private:
+	void InternalLogStreamingKey(uint32 StreamKey, bool const bActive);
 	void InternalLogShader(EShaderPlatform Platform, EShaderFrequency Frequency, FSHAHash Hash, TArray<uint8> const& Code);
 	
 	void InternalLogVertexDeclaration(const FVertexDeclarationElementList& VertexElements, FVertexDeclarationRHIParamRef VertexDeclaration);
@@ -797,8 +827,15 @@ private:
 	TMap<FTextureRHIParamRef, int32> Textures;
 	TMap<FShaderResourceViewRHIParamRef, FShaderResourceKey> SRVs;
 	
+	// Active streaming keys
+	TSet<uint32> ActiveStreamingKeys;
+	
+	// Current combination of streaming keys that define the current streaming environment.
+	// Logged draws will only be predrawn when this key becomes active.
+	uint32 StreamingKey;
+	
 	// Shaders we need to predraw
-	TMap<FShaderCacheBoundState, TSet<int32>> ShadersToDraw;
+	TMap<uint32, FShaderStreamingCache> ShadersToDraw;
 	
 	// Caches to track application & predraw created textures/SRVs so that we minimise temporary resource creation
 	TMap<FShaderTextureKey, FTextureRHIParamRef> CachedTextures;
