UE4 - Fix for ES3.1 emulation on PC

diff --git a/Engine/Source/Runtime/ShaderCore/Private/Shader.cpp b/Engine/Source/Runtime/ShaderCore/Private/Shader.cpp
--- a/Engine/Source/Runtime/ShaderCore/Private/Shader.cpp
+++ b/Engine/Source/Runtime/ShaderCore/Private/Shader.cpp
@@ -438,10 +438,12 @@ bool FShaderResource::ArePlatformsCompatible(EShaderPlatform CurrentPlatform, ES
 
 		bool bIsTargetD3D = TargetPlatform == SP_PCD3D_SM5 ||
 								TargetPlatform == SP_PCD3D_SM4 ||
+								TargetPlatform == SP_PCD3D_ES3_1 ||
 								TargetPlatform == SP_PCD3D_ES2;
 
 		bool bIsCurrentPlatformD3D = CurrentPlatform == SP_PCD3D_SM5 ||
 								CurrentPlatform == SP_PCD3D_SM4 ||
+								TargetPlatform == SP_PCD3D_ES3_1 ||
 								CurrentPlatform == SP_PCD3D_ES2;
 
 		bFeatureLevelCompatible = bFeatureLevelCompatible && (bIsCurrentPlatformD3D == bIsTargetD3D);
