More explanatory code documentation for FShaderCache.

diff --git a/Engine/Source/Runtime/ShaderCore/Public/ShaderCache.h b/Engine/Source/Runtime/ShaderCore/Public/ShaderCache.h
--- a/Engine/Source/Runtime/ShaderCore/Public/ShaderCache.h
+++ b/Engine/Source/Runtime/ShaderCore/Public/ShaderCache.h
@@ -169,7 +169,7 @@ struct SHADERCORE_API FShaderRenderTargetKey
 /**
  * The FShaderCache provides mechanisms for reducing shader hitching in-game:
  *	- Early submission during shader deserialisation rather than on-demand (r.UseShaderCaching 0/1).
- *	- Tracking of bound-shader-states so that they may be pre-bound during early submission.
+ *	- Tracking of bound-shader-states so that they may be pre-bound during early submission (r.UseShaderCaching 0/1).
  *	- Tracking of RHI draw states so that each bound-shader-state (r.UseShaderDrawLog 0/1).
  *	- Predrawing of tracked RHI draw states to eliminate first-use hitches (r.UseShaderPredraw 0/1).
  *	- Control over time spent predrawing each frame (r.PredrawBatchTime Time (ms), -1 for all).
@@ -189,12 +189,15 @@ struct SHADERCORE_API FShaderRenderTargetKey
  * Handling Updates/Invalidation:
  * When the cache needs to be updated & writable caches invalidated the game should specify a new GameVersion.
  * Call FShaderCache::SetGameVersion before initialisating the RHI (which initialises the cache), which will cause 
- * the contents of caches from previous version to be ignored. Presently you cannot carry over cache entries from a previous version.
+ * the contents of caches from previous version to be ignored. At present you cannot carry over cache entries from a previous version.
  *
  * Region/Stream Batching:
  * For streaming games, or where the cache becomes very large, calls to FShaderCache::SetStreamingKey should be added with unique values for
  * the currently relevant game regions/streaming levels (as required). Logged draw states will be linked to the active streaming key.
  * This limits predrawing to only those draw states required by the active streaming key on subsequent runs.
+ * 
+ * Supported RHIs:
+ * - OpenGLDrv
  */
 class SHADERCORE_API FShaderCache
 {
@@ -402,27 +405,36 @@ public:
 	FShaderCache();
 	~FShaderCache();
 	
-	// Called by the game to set the game specific shader cache version, only caches of this version will be loaded.
-	// Must be called before RHI initialisation, as InitShaderCache will load any existing cache.
+	/** Called by the game to set the game specific shader cache version, only caches of this version will be loaded. Must be called before RHI initialisation, as InitShaderCache will load any existing cache. */
 	static void SetGameVersion(int32 InGameVersion);
 	
-	// Called by the RHI implementation, not to be called by the game.
+	/** Shader cache initialisation, called only by the RHI. */
 	static void InitShaderCache();
+	/** Shader cache shutdown, called only by the RHI. */
 	static void ShutdownShaderCache();
 	
+	/** Get the global shader cache if it exists or nullptr otherwise. */
 	static FORCEINLINE FShaderCache* GetShaderCache()
 	{
 		return bUseShaderCaching ? Cache : nullptr;
 	}
 	
+	/** Instantiate or retrieve a vertex shader from the cache for the provided code & hash. */
 	FVertexShaderRHIRef GetVertexShader(EShaderPlatform Platform, FSHAHash Hash, TArray<uint8> const& Code);
+	/** Instantiate or retrieve a pixel shader from the cache for the provided code & hash. */
 	FPixelShaderRHIRef GetPixelShader(EShaderPlatform Platform, FSHAHash Hash, TArray<uint8> const& Code);
+	/** Instantiate or retrieve a geometry shader from the cache for the provided code & hash. */
 	FGeometryShaderRHIRef GetGeometryShader(EShaderPlatform Platform, FSHAHash Hash, TArray<uint8> const& Code);
+	/** Instantiate or retrieve a hull shader from the cache for the provided code & hash. */
 	FHullShaderRHIRef GetHullShader(EShaderPlatform Platform, FSHAHash Hash, TArray<uint8> const& Code);
+	/** Instantiate or retrieve a domain shader from the cache for the provided code & hash. */
 	FDomainShaderRHIRef GetDomainShader(EShaderPlatform Platform, FSHAHash Hash, TArray<uint8> const& Code);
+	/** Instantiate or retrieve a compute shader from the cache for the provided code & hash. */
 	FComputeShaderRHIRef GetComputeShader(EShaderPlatform Platform, TArray<uint8> const& Code);
+	/** Instantiate or retrieve a vertex declaration from the cache for the provided layout. */
 	FVertexDeclarationRHIParamRef GetVertexDeclaration(FVertexDeclarationElementList& VertexElements);
 	
+	/** Called by the game. Logs whether a user-defined streaming key is active or disabled, which is used to batch logged draw-states so that they are only predrawn when the same combination of streaming keys are active. This is important for streaming games, or those with very large numbers of shaders, as it will reduce the number of predraw requests performed at any given time. */
 	static FORCEINLINE void LogStreamingKey(uint32 StreamingKey, bool const bActive)
 	{
 		if ( Cache )
@@ -431,6 +443,7 @@ public:
 		}
 	}
 	
+	/** Called by the Engine/RHI. Logs the use of a given shader & will ensure it is instantiated if not already. */
 	static FORCEINLINE void LogShader(EShaderPlatform Platform, EShaderFrequency Frequency, FSHAHash Hash, TArray<uint8> const& Code)
 	{
 		if ( Cache )
@@ -439,6 +452,7 @@ public:
 		}
 	}
 	
+	/** Called by the RHI. Logs the user of a vertex declaration so that the cache can find it for recording shader & draw states. */
 	static FORCEINLINE void LogVertexDeclaration(const FVertexDeclarationElementList& VertexElements, FVertexDeclarationRHIParamRef VertexDeclaration)
 	{
 		if ( Cache )
@@ -447,6 +461,7 @@ public:
 		}
 	}
 	
+	/** Called by the RHI. Logs the construction of a bound shader state & will record it for prebinding on subsequent runs. */
 	static FORCEINLINE void LogBoundShaderState(EShaderPlatform Platform, FVertexDeclarationRHIParamRef VertexDeclaration,
 												FVertexShaderRHIParamRef VertexShader,
 												FPixelShaderRHIParamRef PixelShader,
@@ -461,6 +476,7 @@ public:
 		}
 	}
 	
+	/** Called by the RHI. Logs the construction of a blend state & caches it so that draw states can be properly recorded. No-op when r.UseShaderDrawLog & r.UseShaderPredraw are disabled. */
 	static FORCEINLINE void LogBlendState(FBlendStateInitializerRHI const& Init, FBlendStateRHIParamRef State)
 	{
 		if ( Cache )
@@ -468,6 +484,8 @@ public:
 			Cache->InternalLogBlendState(Init, State);
 		}
 	}
+	
+	/** Called by the RHI. Logs the construction of a rasterize state & caches it so that draw states can be properly recorded. No-op when r.UseShaderDrawLog & r.UseShaderPredraw are disabled. */
 	static FORCEINLINE void LogRasterizerState(FRasterizerStateInitializerRHI const& Init, FRasterizerStateRHIParamRef State)
 	{
 		if ( Cache )
@@ -475,6 +493,8 @@ public:
 			Cache->InternalLogRasterizerState(Init, State);
 		}
 	}
+	
+	/** Called by the RHI. Logs the construction of a stencil state & caches it so that draw states can be properly recorded. No-op when r.UseShaderDrawLog & r.UseShaderPredraw are disabled. */
 	static FORCEINLINE void LogDepthStencilState(FDepthStencilStateInitializerRHI const& Init, FDepthStencilStateRHIParamRef State)
 	{
 		if ( Cache )
@@ -482,6 +502,8 @@ public:
 			Cache->InternalLogDepthStencilState(Init, State);
 		}
 	}
+	
+	/** Called by the RHI. Logs the construction of a sampelr state & caches it so that draw states can be properly recorded. No-op when r.UseShaderDrawLog & r.UseShaderPredraw are disabled. */
 	static FORCEINLINE void LogSamplerState(FSamplerStateInitializerRHI const& Init, FSamplerStateRHIParamRef State)
 	{
 		if ( Cache )
@@ -489,6 +511,8 @@ public:
 			Cache->InternalLogSamplerState(Init, State);
 		}
 	}
+	
+	/** Called by the RHI. Logs the construction of a texture & caches it so that draw states can be properly recorded. No-op when r.UseShaderDrawLog & r.UseShaderPredraw are disabled. */
 	static FORCEINLINE void LogTexture(FShaderTextureKey const& Init, FTextureRHIParamRef State)
 	{
 		if ( Cache )
@@ -496,6 +520,8 @@ public:
 			Cache->InternalLogTexture(Init, State);
 		}
 	}
+	
+	/** Called by the RHI. Logs the construction of a SRV & caches it so that draw states can be properly recorded. No-op when r.UseShaderDrawLog & r.UseShaderPredraw are disabled. */
 	static FORCEINLINE void LogSRV(FShaderResourceViewRHIParamRef SRV, FTextureRHIParamRef Texture, uint8 StartMip, uint8 NumMips, uint8 Format)
 	{
 		if ( Cache )
@@ -503,6 +529,8 @@ public:
 			Cache->InternalLogSRV(SRV, Texture, StartMip, NumMips, Format);
 		}
 	}
+	
+	/** Called by the RHI. Logs the construction of a SRV & caches it so that draw states can be properly recorded. No-op when r.UseShaderDrawLog & r.UseShaderPredraw are disabled. */
 	static FORCEINLINE void LogSRV(FShaderResourceViewRHIParamRef SRV, FVertexBufferRHIParamRef Vb, uint32 Stride, uint8 Format)
 	{
 		if ( Cache )
@@ -510,6 +538,8 @@ public:
 			Cache->InternalLogSRV(SRV, Vb, Stride, Format);
 		}
 	}
+	
+	/** Called by the RHI. Removes the cached SRV pointer on destruction to prevent further use. No-op when r.UseShaderDrawLog & r.UseShaderPredraw are disabled. */
 	static FORCEINLINE void RemoveSRV(FShaderResourceViewRHIParamRef SRV)
 	{
 		if ( Cache )
@@ -517,6 +547,8 @@ public:
 			Cache->InternalRemoveSRV(SRV);
 		}
 	}
+	
+	/** Called by the RHI. Removes the cached texture pointer on destruction to prevent further use. No-op when r.UseShaderDrawLog & r.UseShaderPredraw are disabled. */
 	static FORCEINLINE void RemoveTexture(FTextureRHIParamRef Texture)
 	{
 		if ( Cache )
@@ -524,6 +556,8 @@ public:
 			Cache->InternalRemoveTexture(Texture);
 		}
 	}
+	
+	/** Called by the RHI. Records the current blend state when r.UseShaderDrawLog or r.UseShaderPredraw are enabled. */
 	static FORCEINLINE void SetBlendState(FBlendStateRHIParamRef State)
 	{
 		if ( Cache )
@@ -531,6 +565,8 @@ public:
 			Cache->InternalSetBlendState(State);
 		}
 	}
+	
+	/** Called by the RHI. Records the current rasterizer state when r.UseShaderDrawLog or r.UseShaderPredraw are enabled. */
 	static FORCEINLINE void SetRasterizerState(FRasterizerStateRHIParamRef State)
 	{
 		if ( Cache )
@@ -538,6 +574,8 @@ public:
 			Cache->InternalSetRasterizerState(State);
 		}
 	}
+	
+	/** Called by the RHI. Records the current depth stencil state when r.UseShaderDrawLog or r.UseShaderPredraw are enabled. */
 	static FORCEINLINE void SetDepthStencilState(FDepthStencilStateRHIParamRef State)
 	{
 		if ( Cache )
@@ -545,6 +583,8 @@ public:
 			Cache->InternalSetDepthStencilState(State);
 		}
 	}
+	
+	/** Called by the RHI. Records the current render-targets when r.UseShaderDrawLog or r.UseShaderPredraw are enabled. */
 	static FORCEINLINE void SetRenderTargets( uint32 NumSimultaneousRenderTargets, const FRHIRenderTargetView* NewRenderTargetsRHI, const FRHIDepthRenderTargetView* NewDepthStencilTargetRHI )
 	{
 		if ( Cache )
@@ -552,6 +592,8 @@ public:
 			Cache->InternalSetRenderTargets(NumSimultaneousRenderTargets, NewRenderTargetsRHI, NewDepthStencilTargetRHI);
 		}
 	}
+	
+	/** Called by the RHI. Records the current sampler state when r.UseShaderDrawLog or r.UseShaderPredraw are enabled. */
 	static FORCEINLINE void SetSamplerState(EShaderFrequency Frequency, uint32 Index, FSamplerStateRHIParamRef State)
 	{
 		if ( Cache )
@@ -559,6 +601,8 @@ public:
 			Cache->InternalSetSamplerState(Frequency, Index, State);
 		}
 	}
+	
+	/** Called by the RHI. Records the current texture binding when r.UseShaderDrawLog or r.UseShaderPredraw are enabled. */
 	static FORCEINLINE void SetTexture(EShaderFrequency Frequency, uint32 Index, FTextureRHIParamRef State)
 	{
 		if ( Cache )
@@ -566,6 +610,8 @@ public:
 			Cache->InternalSetTexture(Frequency, Index, State);
 		}
 	}
+	
+	/** Called by the RHI. Records the current SRV binding when r.UseShaderDrawLog or r.UseShaderPredraw are enabled. */
 	static FORCEINLINE void SetSRV(EShaderFrequency Frequency, uint32 Index, FShaderResourceViewRHIParamRef SRV)
 	{
 		if ( Cache )
@@ -573,6 +619,8 @@ public:
 			Cache->InternalSetSRV(Frequency, Index, SRV);
 		}
 	}
+	
+	/** Called by the RHI. Records the current bound shader state when r.UseShaderDrawLog or r.UseShaderPredraw are enabled. */
 	static FORCEINLINE void SetBoundShaderState(FBoundShaderStateRHIParamRef State)
 	{
 		if ( Cache )
@@ -580,6 +628,8 @@ public:
 			Cache->InternalSetBoundShaderState(State);
 		}
 	}
+	
+	/** Called by the RHI. Records the current vieport when r.UseShaderDrawLog or r.UseShaderPredraw are enabled. */
 	static FORCEINLINE void SetViewport(uint32 MinX, uint32 MinY, float MinZ, uint32 MaxX, uint32 MaxY, float MaxZ)
 	{
 		if ( Cache )
@@ -587,6 +637,8 @@ public:
 			Cache->InternalSetViewport(MinX, MinY, MinZ, MaxX, MaxY, MaxZ);
 		}
 	}
+	
+	/** Called by the RHI. Records the current draw state using the information captured from the other Log/Set* calls if and only if r.UseShaderCaching & r.UseShaderDrawLog are enabled. */
 	static FORCEINLINE void LogDraw(uint8 IndexType)
 	{
 		if ( Cache )
@@ -594,6 +646,8 @@ public:
 			Cache->InternalLogDraw(IndexType);
 		}
 	}
+	
+	/** Called by the RHI. Predraws a batch of shaders if and only if there are any outstanding from the current shader cache and r.UseShaderCaching & r.UseShaderPredraw are enabled. */
 	static FORCEINLINE void PreDrawShaders(FRHICommandList& RHICmdList)
 	{
 		if ( Cache )
@@ -602,6 +656,7 @@ public:
 		}
 	}
 	
+	/** Archive serialisation of the cache data. */
 	friend FArchive& operator<<( FArchive& Ar, FShaderCache& Info );
 	
 private:
