Correct incomplete sentence in FShaderCache comment.

diff --git a/Engine/Source/Runtime/ShaderCore/Public/ShaderCache.h b/Engine/Source/Runtime/ShaderCore/Public/ShaderCache.h
--- a/Engine/Source/Runtime/ShaderCore/Public/ShaderCache.h
+++ b/Engine/Source/Runtime/ShaderCore/Public/ShaderCache.h
@@ -170,14 +170,14 @@ struct SHADERCORE_API FShaderRenderTargetKey
  * The FShaderCache provides mechanisms for reducing shader hitching in-game:
  *	- Early submission during shader deserialisation rather than on-demand (r.UseShaderCaching 0/1).
  *	- Tracking of bound-shader-states so that they may be pre-bound during early submission (r.UseShaderCaching 0/1).
- *	- Tracking of RHI draw states so that each bound-shader-state (r.UseShaderDrawLog 0/1).
+ *	- Tracking of RHI draw states so that each bound-shader-state can be predrawn (r.UseShaderDrawLog 0/1).
  *	- Predrawing of tracked RHI draw states to eliminate first-use hitches (r.UseShaderPredraw 0/1).
- *	- Control over time spent predrawing each frame (r.PredrawBatchTime Time (ms), -1 for all).
+ *	- Control over time spent predrawing each frame to distribute over many frames if required (r.PredrawBatchTime Time (ms), -1 for all).
  *
  * The cache should be populated by enabling r.UseShaderCaching & r.UseShaderDrawLog on a development machine. 
  * Users/players should then consume the cache by enabling r.UseShaderCaching & r.UseShaderPredraw. 
  * Draw logging (r.UseShaderDrawLog) adds noticeable fixed overhead so should be avoided if possible. 
- * Since the caching is done via shader hashes it is advisable to only use this as a final optimisation tool 
+ * Since the caching is done via shader hashes it is also advisable to only use this as a final optimisation tool 
  * when content is largely complete as changes to shader hashes will result in unusued entries accumulating in the cache, 
  * increasing cache size without reducing hitches.
  * 
