[INTEGRATE] Change 2527228 by [REDACTED_PERSON][REDACTED_HANDLE] on [REDACTED_PHONE]:09

Fix for FShader crash while streaming a level.

- Made FShader ref counting more thread safe
- Added extra locks when unregistering the shader so that released shaders can't be found in the shader id map while a shader is being released
- Made sure FShaderType::FindShaderById adds a reference to existing shaders while serializaing inline shader maps.

[INTEGRATE] Change 2532254 by [REDACTED_PERSON][REDACTED_HANDLE] on [REDACTED_PHONE]:18

Better fix for crashes when streaming. Making FShader AND FShaderResource handling thread safe

diff --git a/Engine/Source/Runtime/ShaderCore/Private/Shader.cpp b/Engine/Source/Runtime/ShaderCore/Private/Shader.cpp
--- a/Engine/Source/Runtime/ShaderCore/Private/Shader.cpp
+++ b/Engine/Source/Runtime/ShaderCore/Private/Shader.cpp
@@ -230,10 +230,10 @@ FArchive& operator<<(FArchive& Ar,FShaderType*& Ref)
 }
 
 
-FShader* FShaderType::FindShaderById(const FShaderId& Id)
+TRefCountPtr<FShader> FShaderType::FindShaderById(const FShaderId& Id)
 {
 	FScopeLock MapLock(&ShaderIdMapCritical);
-	FShader* Result = ShaderIdMap.FindRef(Id);
+	TRefCountPtr<FShader> Result = ShaderIdMap.FindRef(Id);
 	return Result;
 }
 
@@ -374,20 +374,22 @@ void FShaderResource::Serialize(FArchive& Ar)
 
 void FShaderResource::AddRef()
 {
-	check(Canary != FShader::ShaderMagic_CleaningUp);
+	// Lock shader id map to prevent anything from acquiring shaders while we manipulate their references
+	FScopeLock ShaderResourceIdMapLock(&ShaderResourceIdMapCritical);
+	check(Canary != FShader::ShaderMagic_CleaningUp);	
 	++NumRefs;
 }
 
 
 void FShaderResource::Release()
 {
+	// We need to lock the resource map so that no resource gets acquired by
+	// FindShaderResourceById while we (potentially) remove this resource
+	FScopeLock ShaderResourceIdMapLock(&ShaderResourceIdMapCritical);
 	check(NumRefs != 0);
 	if(--NumRefs == 0)
 	{
-		{
-			FScopeLock ShaderResourceIdMapLock(&ShaderResourceIdMapCritical);
-			ShaderResourceIdMap.Remove(GetId());
-		}
+		ShaderResourceIdMap.Remove(GetId());
 
 		// Send a release message to the rendering thread when the shader loses its last reference.
 		BeginReleaseResource(this);
@@ -398,10 +400,11 @@ void FShaderResource::Release()
 }
 
 
-FShaderResource* FShaderResource::FindShaderResourceById(const FShaderResourceId& Id)
+TRefCountPtr<FShaderResource> FShaderResource::FindShaderResourceById(const FShaderResourceId& Id)
 {
 	FScopeLock ShaderResourceIdMapLock(&ShaderResourceIdMapCritical);
-	return ShaderResourceIdMap.FindRef(Id);
+	TRefCountPtr<FShaderResource> Result = ShaderResourceIdMap.FindRef(Id);
+	return Result;
 }
 
 
@@ -898,7 +901,7 @@ bool FShader::SerializeBase(FArchive& Ar, bool bShadersInline)
 			FShaderResource* ShaderResource = new FShaderResource();
 			ShaderResource->Serialize(Ar);
 
-			FShaderResource* ExistingResource = FShaderResource::FindShaderResourceById(ShaderResource->GetId());
+			TRefCountPtr<FShaderResource> ExistingResource = FShaderResource::FindShaderResourceById(ShaderResource->GetId());
 
 			// Reuse an existing shader resource if a matching one already exists in memory
 			if (ExistingResource)
@@ -926,7 +929,7 @@ bool FShader::SerializeBase(FArchive& Ar, bool bShadersInline)
 			ResourceId.OutputHash = OutputHash;
 
 			// use it to look up in the registered resource map
-			FShaderResource* ExistingResource = FShaderResource::FindShaderResourceById(ResourceId);
+			TRefCountPtr<FShaderResource> ExistingResource = FShaderResource::FindShaderResourceById(ResourceId);
 			SetResource(ExistingResource);
 		}
 	}
@@ -935,20 +938,26 @@ bool FShader::SerializeBase(FArchive& Ar, bool bShadersInline)
 }
 
 void FShader::AddRef()
-{
+{	
 	check(Canary != ShaderMagic_CleaningUp);
+	// Lock shader Id maps
+	LockShaderIdMap();
 	++NumRefs;
 	if (NumRefs == 1)
 	{
 		INC_DWORD_STAT_BY(STAT_Shaders_ShaderMemory, GetSizeBytes());
 		INC_DWORD_STAT_BY(STAT_Shaders_NumShadersLoaded,1);
 	}
+	UnlockShaderIdMap();
 }
 
 
 void FShader::Release()
 {
-	check(NumRefs != 0);
+	// Lock the shader id map. Note that we don't necessarily have to deregister at this point but
+	// the shader id map has to be locked while we remove references to this shader so that nothing
+	// can find the shader in the map after we remove the final reference but before we deregister the shader
+	LockShaderIdMap();
 	if(--NumRefs == 0)
 	{
 		DEC_DWORD_STAT_BY(STAT_Shaders_ShaderMemory, GetSizeBytes());
@@ -960,6 +969,7 @@ void FShader::Release()
 		Canary = ShaderMagic_CleaningUp;
 		BeginCleanup(this);
 	}
+	UnlockShaderIdMap();
 }
 
 
@@ -972,12 +982,20 @@ void FShader::Register()
 	Type->AddToShaderIdMap(ShaderId, this);
 }
 
+void FShader::LockShaderIdMap()
+{
+	Type->LockShaderIdMap();
+}
 
 void FShader::Deregister()
 {
 	Type->RemoveFromShaderIdMap(GetId());
 }
 
+void FShader::UnlockShaderIdMap()
+{
+	Type->UnlockShaderIdMap();
+}
 
 FShaderId FShader::GetId() const
 {
diff --git a/Engine/Source/Runtime/ShaderCore/Public/Shader.h b/Engine/Source/Runtime/ShaderCore/Public/Shader.h
--- a/Engine/Source/Runtime/ShaderCore/Public/Shader.h
+++ b/Engine/Source/Runtime/ShaderCore/Public/Shader.h
@@ -160,7 +160,7 @@ public:
 	virtual void FinishCleanup();
 
 	/** Finds a matching shader resource in memory if possible. */
-	SHADERCORE_API static FShaderResource* FindShaderResourceById(const FShaderResourceId& Id);
+	SHADERCORE_API static TRefCountPtr<FShaderResource> FindShaderResourceById(const FShaderResourceId& Id);
 
 	/** 
 	 * Finds a matching shader resource in memory or creates a new one with the given compiler output.  
@@ -655,6 +655,12 @@ protected:
 
 private:
 	
+	/** Locks the type's shader id map so that no other thread can add or remove shaders while we're deregistering. */
+	void LockShaderIdMap();
+
+	/** Unlocks the shader type's id map. */
+	void UnlockShaderIdMap();
+
 	/** 
 	 * Hash of the compiled output from this shader and the resulting parameter map.  
 	 * This is used to find a matching resource.
@@ -746,7 +752,7 @@ public:
 	 * Finds a shader of this type by ID.
 	 * @return NULL if no shader with the specified ID was found.
 	 */
-	FShader* FindShaderById(const FShaderId& Id);
+	TRefCountPtr<FShader> FindShaderById(const FShaderId& Id);
 
 	/** Constructs a new instance of the shader type for deserialization. */
 	FShader* ConstructForDeserialization() const;
@@ -819,12 +825,23 @@ public:
 		ShaderIdMap.Add(Id, Shader);
 	}
 
+	/** Locks the ShaderIdMap before deregistration */
+	void LockShaderIdMap()
+	{
+		ShaderIdMapCritical.Lock();
+	}
+
 	void RemoveFromShaderIdMap(FShaderId Id)
 	{
-		FScopeLock MapLock(&ShaderIdMapCritical);
 		ShaderIdMap.Remove(Id);
 	}
 
+	/** Unlocks the ShaderIdMap after deregistration has completed */
+	void UnlockShaderIdMap()
+	{
+		ShaderIdMapCritical.Unlock();
+	}
+
 	bool LimitShaderResourceToThisType()
 	{
 		return GetStreamOutElementsRef != &FShader::GetStreamOutElements;
@@ -1140,12 +1157,12 @@ public:
 					check(Shader != NULL);
 					Shader->SerializeBase(Ar, bInlineShaderResource);
 
-					FShader* ExistingShader = Type->FindShaderById(Shader->GetId());
+					TRefCountPtr<FShader> ExistingShader = Type->FindShaderById(Shader->GetId());
 
-					if (ExistingShader)
+					if (ExistingShader.IsValid())
 					{
 						delete Shader;
-						Shader = ExistingShader;
+						Shader = ExistingShader.GetReference();
 					}
 					else
 					{
