Fix for inconsistent crash when requesting shader recompilation via cook on the fly.
UE-16572

diff --git a/Engine/Source/Runtime/ShaderCore/Private/Shader.cpp b/Engine/Source/Runtime/ShaderCore/Private/Shader.cpp
--- a/Engine/Source/Runtime/ShaderCore/Private/Shader.cpp
+++ b/Engine/Source/Runtime/ShaderCore/Private/Shader.cpp
@@ -927,6 +927,7 @@ bool FShader::SerializeBase(FArchive& Ar, bool bShadersInline)
 			FShaderResourceId ResourceId;
 			ResourceId.Target = Target;
 			ResourceId.OutputHash = OutputHash;
+			ResourceId.SpecificShaderTypeName = Type->LimitShaderResourceToThisType() ? Type->GetName() : NULL;
 
 			// use it to look up in the registered resource map
 			TRefCountPtr<FShaderResource> ExistingResource = FShaderResource::FindShaderResourceById(ResourceId);
