minor optimizations

diff --git a/Engine/Source/Runtime/ShaderCore/Public/Shader.h b/Engine/Source/Runtime/ShaderCore/Public/Shader.h
--- a/Engine/Source/Runtime/ShaderCore/Public/Shader.h
+++ b/Engine/Source/Runtime/ShaderCore/Public/Shader.h
@@ -581,7 +581,7 @@ public:
 		FUniformBufferStruct* SearchStruct = &UniformBufferStructType::StaticStruct;
 		int32 FoundIndex = INDEX_NONE;
 
-		for (int32 StructIndex = 0; StructIndex < UniformBufferParameterStructs.Num(); StructIndex++)
+		for (int32 StructIndex = 0, Count = UniformBufferParameterStructs.Num(); StructIndex < Count; StructIndex++)
 		{
 			if (UniformBufferParameterStructs[StructIndex] == SearchStruct)
 			{
@@ -612,7 +612,7 @@ public:
 	{
 		int32 FoundIndex = INDEX_NONE;
 
-		for (int32 StructIndex = 0; StructIndex < UniformBufferParameterStructs.Num(); StructIndex++)
+		for (int32 StructIndex = 0, Count = UniformBufferParameterStructs.Num(); StructIndex < Count; StructIndex++)
 		{
 			if (UniformBufferParameterStructs[StructIndex] == SearchStruct)
 			{
@@ -1185,7 +1185,7 @@ public:
 	/** @return true if the map is empty */
 	bool IsEmpty()
 	{
-		return ((Shaders.Num() == 0) ? true : false);
+		return Shaders.Num() == 0;
 	}
 
 	/** @return The number of shaders in the map. */
