UE4 - Fix minor inconsistencies on shader api

diff --git a/Engine/Source/Runtime/Engine/Private/StaticBoundShaderState.cpp b/Engine/Source/Runtime/Engine/Private/StaticBoundShaderState.cpp
--- a/Engine/Source/Runtime/Engine/Private/StaticBoundShaderState.cpp
+++ b/Engine/Source/Runtime/Engine/Private/StaticBoundShaderState.cpp
@@ -125,7 +125,7 @@ static FBoundShaderStateRHIParamRef GetGlobalBoundShaderState_Internal(FGlobalBo
 		WorkArea->Args.VertexDeclarationRHI,
 		GETSAFERHISHADER_VERTEX(WorkArea->Args.VertexShader),
 		GETSAFERHISHADER_PIXEL(WorkArea->Args.PixelShader),
-		(FGeometryShaderRHIParamRef)GETSAFERHISHADER_GEOMETRY(WorkArea->Args.GeometryShader));
+		GETSAFERHISHADER_GEOMETRY(WorkArea->Args.GeometryShader));
 	if (bNewBSS)
 	{
 		FPlatformMisc::MemoryBarrier();
diff --git a/Engine/Source/Runtime/RHI/Public/RHI.h b/Engine/Source/Runtime/RHI/Public/RHI.h
--- a/Engine/Source/Runtime/RHI/Public/RHI.h
+++ b/Engine/Source/Runtime/RHI/Public/RHI.h
@@ -1179,10 +1179,10 @@ extern RHI_API void RHIExit();
 // the following helper macros allow to safely convert shader types without much code clutter
 #define GETSAFERHISHADER_PIXEL(Shader) (Shader ? Shader->GetPixelShader() : (FPixelShaderRHIParamRef)FPixelShaderRHIRef())
 #define GETSAFERHISHADER_VERTEX(Shader) (Shader ? Shader->GetVertexShader() : (FVertexShaderRHIParamRef)FVertexShaderRHIRef())
-#define GETSAFERHISHADER_HULL(Shader) (Shader ? Shader->GetHullShader() : FHullShaderRHIRef())
-#define GETSAFERHISHADER_DOMAIN(Shader) (Shader ? Shader->GetDomainShader() : FDomainShaderRHIRef())
-#define GETSAFERHISHADER_GEOMETRY(Shader) (Shader ? Shader->GetGeometryShader() : FGeometryShaderRHIRef())
-#define GETSAFERHISHADER_COMPUTE(Shader) (Shader ? Shader->GetComputeShader() : FComputeShaderRHIRef())
+#define GETSAFERHISHADER_HULL(Shader) (Shader ? Shader->GetHullShader() : (FHullShaderRHIParamRef)FHullShaderRHIRef())
+#define GETSAFERHISHADER_DOMAIN(Shader) (Shader ? Shader->GetDomainShader() : (FDomainShaderRHIParamRef)FDomainShaderRHIRef())
+#define GETSAFERHISHADER_GEOMETRY(Shader) (Shader ? Shader->GetGeometryShader() : (FGeometryShaderRHIParamRef)FGeometryShaderRHIRef())
+#define GETSAFERHISHADER_COMPUTE(Shader) (Shader ? Shader->GetComputeShader() : (FComputeShaderRHIParamRef)FComputeShaderRHIRef())
 
 // RHI utility functions that depend on the RHI definitions.
 #include "RHIUtilities.h"
diff --git a/Engine/Source/Runtime/ShaderCore/Private/Shader.cpp b/Engine/Source/Runtime/ShaderCore/Private/Shader.cpp
--- a/Engine/Source/Runtime/ShaderCore/Private/Shader.cpp
+++ b/Engine/Source/Runtime/ShaderCore/Private/Shader.cpp
@@ -552,117 +552,8 @@ void FShaderResource::ReleaseRHI()
 	ComputeShader.SafeRelease();
 }
 
-void FShaderResource::InitializeVertexShaderRHI() 
+void FShaderResource::InitializeShaderRHI() 
 { 
-	if (!IsInitialized())
-	{
-		STAT(double ShaderInitializationTime = 0);
-		{
-			SCOPE_CYCLE_COUNTER(STAT_Shaders_FrameRTShaderInitForRenderingTime);
-			SCOPE_SECONDS_COUNTER(ShaderInitializationTime);
-
-			InitResourceFromPossiblyParallelRendering();
-
-		}
-
-		INC_FLOAT_STAT_BY(STAT_Shaders_TotalRTShaderInitForRenderingTime,(float)ShaderInitializationTime);
-	}
-
-	checkSlow(IsInitialized());
-}
-
-
-void FShaderResource::InitializePixelShaderRHI() 
-{ 
-	if (!IsInitialized())
-	{
-		STAT(double ShaderInitializationTime = 0);
-		{
-			SCOPE_CYCLE_COUNTER(STAT_Shaders_FrameRTShaderInitForRenderingTime);
-			SCOPE_SECONDS_COUNTER(ShaderInitializationTime);
-
-			InitResourceFromPossiblyParallelRendering();
-		}
-
-		INC_FLOAT_STAT_BY(STAT_Shaders_TotalRTShaderInitForRenderingTime,(float)ShaderInitializationTime);
-	}
-
-	checkSlow(IsInitialized());
-}
-
-
-const FHullShaderRHIRef& FShaderResource::GetHullShader() 
-{ 
-	checkSlow(Target.Frequency == SF_Hull);
-	if (!IsInitialized())
-	{
-		STAT(double ShaderInitializationTime = 0);
-		{
-			SCOPE_CYCLE_COUNTER(STAT_Shaders_FrameRTShaderInitForRenderingTime);
-			SCOPE_SECONDS_COUNTER(ShaderInitializationTime);
-
-			InitResourceFromPossiblyParallelRendering();
-		}
-
-		INC_FLOAT_STAT_BY(STAT_Shaders_TotalRTShaderInitForRenderingTime,(float)ShaderInitializationTime);
-	}
-
-	checkSlow(IsInitialized());
-
-	return HullShader; 
-}
-
-
-const FDomainShaderRHIRef& FShaderResource::GetDomainShader() 
-{ 
-	checkSlow(Target.Frequency == SF_Domain);
-
-	if (!IsInitialized())
-	{
-		STAT(double ShaderInitializationTime = 0);
-		{
-			SCOPE_CYCLE_COUNTER(STAT_Shaders_FrameRTShaderInitForRenderingTime);
-			SCOPE_SECONDS_COUNTER(ShaderInitializationTime);
-
-			InitResourceFromPossiblyParallelRendering();
-		}
-		
-		INC_FLOAT_STAT_BY(STAT_Shaders_TotalRTShaderInitForRenderingTime,(float)ShaderInitializationTime);
-	}
-	
-	checkSlow(IsInitialized());
-
-	return DomainShader; 
-}
-
-
-const FGeometryShaderRHIRef& FShaderResource::GetGeometryShader() 
-{ 
-	checkSlow(Target.Frequency == SF_Geometry);
-
-	if (!IsInitialized())
-	{
-		STAT(double ShaderInitializationTime = 0);
-		{
-			SCOPE_CYCLE_COUNTER(STAT_Shaders_FrameRTShaderInitForRenderingTime);
-			SCOPE_SECONDS_COUNTER(ShaderInitializationTime);
-
-			InitResourceFromPossiblyParallelRendering();
-		}
-
-		INC_FLOAT_STAT_BY(STAT_Shaders_TotalRTShaderInitForRenderingTime,(float)ShaderInitializationTime);
-	}
-
-	checkSlow(IsInitialized());
-
-	return GeometryShader; 
-}
-
-
-const FComputeShaderRHIRef& FShaderResource::GetComputeShader() 
-{ 
-	checkSlow(Target.Frequency == SF_Compute);
-
 	if (!IsInitialized())
 	{
 		STAT(double ShaderInitializationTime = 0);
@@ -677,8 +568,6 @@ const FComputeShaderRHIRef& FShaderResource::GetComputeShader()
 	}
 
 	checkSlow(IsInitialized());
-
-	return ComputeShader; 
 }
 
 FShaderResourceId FShaderResource::GetId() const
diff --git a/Engine/Source/Runtime/ShaderCore/Public/Shader.h b/Engine/Source/Runtime/ShaderCore/Public/Shader.h
--- a/Engine/Source/Runtime/ShaderCore/Public/Shader.h
+++ b/Engine/Source/Runtime/ShaderCore/Public/Shader.h
@@ -120,9 +120,9 @@ public:
 	FORCEINLINE const FVertexShaderRHIParamRef GetVertexShader()
 	{
 		checkSlow(Target.Frequency == SF_Vertex);
-		if(!IsInitialized())
+		if (!IsInitialized())
 		{
-			InitializeVertexShaderRHI();
+			InitializeShaderRHI();
 		}
 		return VertexShader;
 	}
@@ -130,20 +130,52 @@ public:
 	FORCEINLINE const FPixelShaderRHIParamRef GetPixelShader()
 	{
 		checkSlow(Target.Frequency == SF_Pixel);
-		if(!IsInitialized())
+		if (!IsInitialized())
 		{
-			InitializePixelShaderRHI();
+			InitializeShaderRHI();
 		}
 		return PixelShader;
 	}
 	/** @return the shader's hull shader */
-	const FHullShaderRHIRef& GetHullShader();
+	FORCEINLINE const FHullShaderRHIParamRef GetHullShader()
+	{
+		checkSlow(Target.Frequency == SF_Hull);
+		if (!IsInitialized())
+		{
+			InitializeShaderRHI();
+		}
+		return HullShader;
+	}
 	/** @return the shader's domain shader */
-	const FDomainShaderRHIRef& GetDomainShader();
+	FORCEINLINE const FDomainShaderRHIParamRef GetDomainShader()
+	{
+		checkSlow(Target.Frequency == SF_Domain);
+		if (!IsInitialized())
+		{
+			InitializeShaderRHI();
+		}
+		return DomainShader;
+	}
 	/** @return the shader's geometry shader */
-	const FGeometryShaderRHIRef& GetGeometryShader();
+	FORCEINLINE const FGeometryShaderRHIParamRef GetGeometryShader()
+	{
+		checkSlow(Target.Frequency == SF_Geometry);
+		if (!IsInitialized())
+		{
+			InitializeShaderRHI();
+		}
+		return GeometryShader;
+	}
 	/** @return the shader's compute shader */
-	const FComputeShaderRHIRef& GetComputeShader();
+	FORCEINLINE const FComputeShaderRHIParamRef GetComputeShader()
+	{
+		checkSlow(Target.Frequency == SF_Compute);
+		if (!IsInitialized())
+		{
+			InitializeShaderRHI();
+		}
+		return ComputeShader;
+	}
 
 	SHADERCORE_API FShaderResourceId GetId() const;
 
@@ -211,11 +243,8 @@ private:
 	/** A 'canary' used to detect when a stale shader resource is being rendered with. */
 	uint32 Canary;
 
-	/** Initialize the vertex shader RHI resource. */
-	void InitializeVertexShaderRHI();
-
-	/** Initialize the pixel shader RHI resource. */
-	void InitializePixelShaderRHI();
+	/** Initialize the shader RHI resources. */
+	void InitializeShaderRHI();
 
 	/** Tracks loaded shader resources by id. */
 	static TMap<FShaderResourceId, FShaderResource*> ShaderResourceIdMap;
@@ -502,22 +531,22 @@ public:
 		return Resource->GetPixelShader();
 	}
 	/** @return the shader's hull shader */
-	const FHullShaderRHIRef& GetHullShader()
+	inline const FHullShaderRHIParamRef GetHullShader()
 	{
 		return Resource->GetHullShader();
 	}
 	/** @return the shader's domain shader */
-	const FDomainShaderRHIRef& GetDomainShader()
+	inline const FDomainShaderRHIParamRef GetDomainShader()
 	{
 		return Resource->GetDomainShader();
 	}
 	/** @return the shader's geometry shader */
-	const FGeometryShaderRHIRef& GetGeometryShader()
+	inline const FGeometryShaderRHIParamRef GetGeometryShader()
 	{
 		return Resource->GetGeometryShader();
 	}
 	/** @return the shader's compute shader */
-	const FComputeShaderRHIRef& GetComputeShader()
+	inline const FComputeShaderRHIParamRef GetComputeShader()
 	{
 		return Resource->GetComputeShader();
 	}
