Merge changes to the ShaderCache from Zebra to extend support to Mac Metal.

diff --git a/Engine/Source/Runtime/OpenGLDrv/Private/OpenGLDevice.cpp b/Engine/Source/Runtime/OpenGLDrv/Private/OpenGLDevice.cpp
--- a/Engine/Source/Runtime/OpenGLDrv/Private/OpenGLDevice.cpp
+++ b/Engine/Source/Runtime/OpenGLDrv/Private/OpenGLDevice.cpp
@@ -1230,7 +1230,7 @@ void FOpenGLDynamicRHI::Init()
 	VERIFY_GL_SCOPE();
 
 #if PLATFORM_DESKTOP
-	FShaderCache::InitShaderCache();
+	FShaderCache::InitShaderCache(SCO_Default, FOpenGL::GetMaxTextureImageUnits());
 #endif
 
 	InitializeStateResources();
diff --git a/Engine/Source/Runtime/ShaderCore/Private/ShaderCache.cpp b/Engine/Source/Runtime/ShaderCore/Private/ShaderCache.cpp
--- a/Engine/Source/Runtime/ShaderCore/Private/ShaderCache.cpp
+++ b/Engine/Source/Runtime/ShaderCore/Private/ShaderCache.cpp
@@ -13,9 +13,15 @@
 const FGuid FShaderCacheCustomVersion::Key(0xB954F018, 0xC9624DD6, 0xA74E79B1, 0x8EA113C2);
 const FGuid FShaderCacheCustomVersion::GameKey(0x03D4EB48, 0xB50B4CC3, 0xA598DE41, 0x5C6CC993);
 FCustomVersionRegistration GRegisterShaderCacheVersion(FShaderCacheCustomVersion::Key, FShaderCacheCustomVersion::Latest, TEXT("ShaderCacheVersion"));
+#if WITH_EDITOR
+static TCHAR const* GShaderCacheFileName = TEXT("EditorShaderCache.ushadercache");
+static TCHAR const* GShaderCodeCacheFileName = TEXT("EditorCodeCache.ushadercode");
+#else
 static TCHAR const* GShaderCacheFileName = TEXT("ShaderCache.ushadercache");
+static TCHAR const* GShaderCodeCacheFileName = TEXT("ShaderCodeCache.ushadercode");
+#endif
 
-// Only the cooked Mac build defaults to using the shader cache for now, Editor is too likely to invalidate shader keys leading to ever growing cache
+// Only the Mac build defaults to using the shader cache for now, Editor uses a separate cache from the game to avoid ever-growing cache being propagated to the game.
 int32 FShaderCache::bUseShaderCaching = (PLATFORM_MAC && !WITH_EDITOR) ? 1 : 0;
 FAutoConsoleVariableRef FShaderCache::CVarUseShaderCaching(
 	TEXT("r.UseShaderCaching"),
@@ -35,7 +41,7 @@ FAutoConsoleVariableRef FShaderCache::CVarUseShaderPredraw(
 	);
 
 // The actual draw loggging is even more expensive as it has to cache all the RHI draw state & is disabled by default.
-int32 FShaderCache::bUseShaderDrawLog = 0;
+int32 FShaderCache::bUseShaderDrawLog = (PLATFORM_MAC && !WITH_EDITOR) ? 1 : 0;
 FAutoConsoleVariableRef FShaderCache::CVarUseShaderDrawLog(
 	TEXT("r.UseShaderDrawLog"),
 	bUseShaderDrawLog,
@@ -52,6 +58,15 @@ FAutoConsoleVariableRef FShaderCache::CVarPredrawBatchTime(
 	ECVF_RenderThreadSafe
 	);
 
+// A separate cache of used shader binaries for even earlier submission - may be platform or even device specific.
+int32 FShaderCache::bUseShaderBinaryCache = 0;
+FAutoConsoleVariableRef FShaderCache::CVarUseShaderBinaryCache(
+	TEXT("r.UseShaderBinaryCache"),
+	bUseShaderBinaryCache,
+	TEXT("If true generates & uses a separate cache of used shader binaries for even earlier submission - may be platform or even device specific. Defaults to false."),
+	ECVF_ReadOnly|ECVF_RenderThreadSafe
+	);
+
 FShaderCache* FShaderCache::Cache = nullptr;
 int32 FShaderCache::GameVersion = 0;
 
@@ -67,6 +82,8 @@ static bool ShaderPlatformCanPrebindBoundShaderState(EShaderPlatform Platform)
 		case SP_METAL:
 		case SP_OPENGL_SM4_MAC:
 		case SP_METAL_MRT:
+		case SP_METAL_SM4:
+		case SP_METAL_SM5:
 		{
 			return true;
 		}
@@ -90,12 +107,63 @@ void FShaderCache::SetGameVersion(int32 InGameVersion)
 	GameVersion = InGameVersion;
 }
 
-void FShaderCache::InitShaderCache()
+void FShaderCache::InitShaderCache(EShaderCacheOptions Options, uint32 InMaxResources)
 {
 	check(!Cache);
 	if(bUseShaderCaching)
 	{
 		Cache = new FShaderCache;
+		Cache->Options = Options;
+		Cache->MaxResources = InMaxResources;
+		check(InMaxResources <= FShaderDrawKey::MaxNumResources);
+		
+		if(bUseShaderBinaryCache)
+		{
+			FString UserCodeCache = FPaths::GameSavedDir() / GShaderCodeCacheFileName;
+			FString GameCodeCache = FPaths::GameContentDir() / GShaderCodeCacheFileName;
+			
+			// Try to load user cache, making sure that if we fail version test we still try game-content version.
+			bool bLoadedCache = false;
+			if ( IFileManager::Get().FileSize(*UserCodeCache) > 0 )
+			{
+				FArchive* BinaryShaderAr = IFileManager::Get().CreateFileReader(*UserCodeCache);
+				
+				if ( BinaryShaderAr != nullptr )
+				{
+					*BinaryShaderAr << Cache->CodeCache;
+					
+					if ( !BinaryShaderAr->IsError() && BinaryShaderAr->CustomVer(FShaderCacheCustomVersion::Key) == FShaderCacheCustomVersion::Latest && BinaryShaderAr->CustomVer(FShaderCacheCustomVersion::GameKey) == FShaderCache::GameVersion )
+					{
+						bLoadedCache = true;
+					}
+					
+					delete BinaryShaderAr;
+				}
+			}
+			
+			// Fallback to game-content version.
+			if ( !bLoadedCache && IFileManager::Get().FileSize(*GameCodeCache) > 0 )
+			{
+				FArchive* BinaryShaderAr = IFileManager::Get().CreateFileReader(*GameCodeCache);
+				if ( BinaryShaderAr != nullptr )
+				{
+					*BinaryShaderAr << Cache->CodeCache;
+					bLoadedCache = true;
+					delete BinaryShaderAr;
+				}
+			}
+			
+			if(bLoadedCache)
+			{
+				for(auto Entry : Cache->CodeCache.Shaders)
+				{
+					if(Entry.Key.Platform == GMaxRHIShaderPlatform)
+					{
+						Cache->InternalLogShader(Entry.Key.Platform, Entry.Key.Frequency, Entry.Key.SHAHash, Entry.Value);
+					}
+				}
+			}
+		}
 	}
 }
 
@@ -127,12 +195,32 @@ FArchive& operator<<( FArchive& Ar, FShaderCache& Info )
 	return Ar;
 }
 
+FArchive& operator<<( FArchive& Ar, FShaderCache::FShaderCodeCache& Info )
+{
+	Ar.UsingCustomVersion(FShaderCacheCustomVersion::Key);
+	Ar.SetCustomVersion(FShaderCacheCustomVersion::GameKey, FShaderCache::GetGameVersion(), TEXT("ShaderCacheGameVersion"));
+	
+	FCustomVersionContainer Vers = Ar.GetCustomVersions();
+	Vers.Serialize(Ar);
+	if ( Ar.IsLoading() )
+	{
+		Ar.SetCustomVersions(Vers);
+	}
+	
+	if ( !Ar.IsError() && Ar.CustomVer(FShaderCacheCustomVersion::Key) == FShaderCacheCustomVersion::Latest && Ar.CustomVer(FShaderCacheCustomVersion::GameKey) == FShaderCache::GetGameVersion() )
+	{
+		Ar << Info.Shaders;
+	}
+	return Ar;
+}
+
 FShaderCache::FShaderCache()
 : StreamingKey(0)
 , bCurrentDepthStencilTarget(false)
 , CurrentNumRenderTargets(0)
 , CurrentShaderState(nullptr)
 , bIsPreDraw(false)
+, Options(SCO_Default)
 {
 	Viewport[0] = Viewport[1] = Viewport[2] = Viewport[3] = 0;
 	DepthRange[0] = DepthRange[1] = 0.0f;
@@ -180,6 +268,16 @@ FShaderCache::~FShaderCache()
 		*BinaryShaderAr << *this;
 		delete BinaryShaderAr;
 	}
+	if(bUseShaderBinaryCache)
+	{
+		BinaryShaderFile = FPaths::GameSavedDir() / GShaderCodeCacheFileName;
+		BinaryShaderAr = IFileManager::Get().CreateFileWriter(*BinaryShaderFile);
+		if( BinaryShaderAr != NULL )
+		{
+			*BinaryShaderAr << CodeCache;
+			delete BinaryShaderAr;
+		}
+	}
 }
 
 FVertexShaderRHIRef FShaderCache::GetVertexShader(EShaderPlatform Platform, FSHAHash Hash, TArray<uint8> const& Code)
@@ -370,9 +468,18 @@ void FShaderCache::InternalLogShader(EShaderPlatform Platform, EShaderFrequency
 		Key.Platform = Platform;
 		Key.Frequency = Frequency;
 		Key.bActive = true;
-		ENQUEUE_UNIQUE_RENDER_COMMAND_TWOPARAMETER(LogShader, FShaderCacheKey,Key,Key, TArray<uint8>,Code,Code, {
-			FShaderCache::GetShaderCache()->SubmitShader(Key, Code);
-		});
+		
+		if(bUseShaderBinaryCache && !CodeCache.Shaders.Contains(Key))
+		{
+			CodeCache.Shaders.Add(Key, Code);
+		}
+		
+		if(!(Options & SCO_NoShaderPreload))
+		{
+			ENQUEUE_UNIQUE_RENDER_COMMAND_TWOPARAMETER(LogShader, FShaderCacheKey,Key,Key, TArray<uint8>,Code,Code, {
+				FShaderCache::GetShaderCache()->SubmitShader(Key, Code);
+			});
+		}
 	}
 }
 
@@ -505,7 +612,7 @@ void FShaderCache::InternalLogDepthStencilState(FDepthStencilStateInitializerRHI
 
 void FShaderCache::InternalLogSamplerState(FSamplerStateInitializerRHI const& Init, FSamplerStateRHIParamRef State)
 {
-	if ( bUseShaderDrawLog )
+	if ( bUseShaderPredraw || bUseShaderDrawLog )
 	{
 		FShaderPlatformCache& PlatformCache = Caches.FindOrAdd(GMaxRHIShaderPlatform);
 		FSetElementId ID = PlatformCache.SamplerStates.Add(Init);
@@ -595,7 +702,7 @@ void FShaderCache::InternalRemoveTexture(FTextureRHIParamRef Texture)
 
 void FShaderCache::InternalSetBlendState(FBlendStateRHIParamRef State)
 {
-	if ( (bUseShaderPredraw || bUseShaderDrawLog) && !bIsPreDraw )
+	if ( (bUseShaderPredraw || bUseShaderDrawLog) && !bIsPreDraw && State )
 	{
 		CurrentDrawKey.BlendState = BlendStates.FindChecked(State);
 		CurrentDrawKey.Hash = 0;
@@ -604,7 +711,7 @@ void FShaderCache::InternalSetBlendState(FBlendStateRHIParamRef State)
 
 void FShaderCache::InternalSetRasterizerState(FRasterizerStateRHIParamRef State)
 {
-	if ( (bUseShaderPredraw || bUseShaderDrawLog) && !bIsPreDraw )
+	if ( (bUseShaderPredraw || bUseShaderDrawLog) && !bIsPreDraw && State )
 	{
 		CurrentDrawKey.RasterizerState = RasterizerStates.FindChecked(State);
 		CurrentDrawKey.Hash = 0;
@@ -613,7 +720,7 @@ void FShaderCache::InternalSetRasterizerState(FRasterizerStateRHIParamRef State)
 
 void FShaderCache::InternalSetDepthStencilState(FDepthStencilStateRHIParamRef State)
 {
-	if ( (bUseShaderPredraw || bUseShaderDrawLog) && !bIsPreDraw )
+	if ( (bUseShaderPredraw || bUseShaderDrawLog) && !bIsPreDraw && State )
 	{
 		CurrentDrawKey.DepthStencilState = DepthStencilStates.FindChecked(State);
 		CurrentDrawKey.Hash = 0;
@@ -693,7 +800,7 @@ void FShaderCache::InternalSetTexture(EShaderFrequency Frequency, uint32 Index,
 {
 	if ( bUseShaderDrawLog && !bIsPreDraw )
 	{
-		check(Index < GetFeatureLevelMaxTextureSamplers(GMaxRHIFeatureLevel));
+		check(Index < MaxResources);
 		if ( State )
 		{
 			FShaderResourceKey Key;
@@ -720,7 +827,7 @@ void FShaderCache::InternalSetSRV(EShaderFrequency Frequency, uint32 Index, FSha
 {
 	if ( bUseShaderDrawLog && !bIsPreDraw )
 	{
-		check(Index < GetFeatureLevelMaxTextureSamplers(GMaxRHIFeatureLevel));
+		check(Index < MaxResources);
 		if ( SRV )
 		{
 			FShaderResourceKey Key = SRVs.FindChecked(SRV);
@@ -1091,7 +1198,23 @@ FShaderCache::FShaderTextureBinding FShaderCache::CreateSRV(FShaderResourceKey c
 			case SCTT_Texture2D:
 			{
 				Binding.Texture = CreateTexture(TextureKey, true);
-				Binding.SRV = RHICreateShaderResourceView(Binding.Texture->GetTexture2D(), ResourceKey.BaseMip, ResourceKey.MipLevels, ResourceKey.Format);
+				
+				if(ResourceKey.Format == PF_Unknown)
+				{
+					Binding.SRV = RHICreateShaderResourceView(Binding.Texture->GetTexture2D(), ResourceKey.BaseMip);
+					
+				}
+				else
+				{
+					// Make sure that the mip count is valid.
+					uint32 NumMips = (Binding.Texture->GetNumMips() - ResourceKey.BaseMip);
+					if(ResourceKey.MipLevels > 0)
+					{
+						NumMips = FMath::Min(NumMips, ResourceKey.MipLevels);
+					}
+					
+					Binding.SRV = RHICreateShaderResourceView(Binding.Texture->GetTexture2D(), ResourceKey.BaseMip, NumMips, ResourceKey.Format);
+				}
 				break;
 			}
 			default:
@@ -1132,46 +1255,49 @@ void FShaderCache::SetShaderSamplerTextures( FRHICommandList& RHICmdList, FShade
 			FSamplerStateInitializerRHI SamplerInit = PlatformCache.SamplerStates[FSetElementId::FromInteger(DrawKey.SamplerStates[Frequency][i])];
 			FSamplerStateRHIRef State = RHICreateSamplerState(SamplerInit);
 			RHICmdList.SetShaderSampler(Shader, i, State);
-			
-			FShaderTextureBinding Bind;
-			if ( DrawKey.Resources[Frequency][i] != FShaderDrawKey::NullState )
+		}
+	}
+
+	for ( uint32 i = 0; i < MaxResources; i++ )
+	{
+		FShaderTextureBinding Bind;
+		if ( DrawKey.Resources[Frequency][i] != FShaderDrawKey::NullState )
+		{
+			FShaderResourceKey Resource = PlatformCache.Resources[FSetElementId::FromInteger(DrawKey.Resources[Frequency][i])];
+			if( Resource.bSRV == false )
 			{
-				FShaderResourceKey Resource = PlatformCache.Resources[FSetElementId::FromInteger(DrawKey.Resources[Frequency][i])];
-				if( Resource.bSRV == false )
+				if ( !bClear && Resource.Tex.Type != SCTT_Invalid )
 				{
-					if ( !bClear && Resource.Tex.Type != SCTT_Invalid )
-					{
-						Bind.Texture = CreateTexture(Resource.Tex, true);
-						RHICmdList.SetShaderTexture(Shader, i, Bind.Texture.GetReference());
-					}
-					else
-					{
-						RHICmdList.SetShaderTexture(Shader, i, nullptr);
-					}
+					Bind.Texture = CreateTexture(Resource.Tex, true);
+					RHICmdList.SetShaderTexture(Shader, i, Bind.Texture.GetReference());
 				}
 				else
 				{
-					if ( !bClear )
-					{
-						Bind = CreateSRV(Resource);
-						RHICmdList.SetShaderResourceViewParameter(Shader, i, Bind.SRV.GetReference());
-					}
-					else
-					{
-						RHICmdList.SetShaderResourceViewParameter(Shader, i, nullptr);
-					}
+					RHICmdList.SetShaderTexture(Shader, i, nullptr);
 				}
 			}
 			else
 			{
-				RHICmdList.SetShaderTexture(Shader, i, nullptr);
-			}
-			
-			if ( IsValidRef(Bind.Texture) || IsValidRef(Bind.SRV) )
-			{
-				PredrawBindings.Add(Bind);
+				if ( !bClear )
+				{
+					Bind = CreateSRV(Resource);
+					RHICmdList.SetShaderResourceViewParameter(Shader, i, Bind.SRV.GetReference());
+				}
+				else
+				{
+					RHICmdList.SetShaderResourceViewParameter(Shader, i, nullptr);
+				}
 			}
 		}
+		else
+		{
+			RHICmdList.SetShaderTexture(Shader, i, nullptr);
+		}
+		
+		if ( IsValidRef(Bind.Texture) || IsValidRef(Bind.SRV) )
+		{
+			PredrawBindings.Add(Bind);
+		}
 	}
 }
 
@@ -1212,18 +1338,6 @@ void FShaderCache::PreDrawShader(FRHICommandList& RHICmdList, FShaderCacheBoundS
 			}
 		}
 		
-		for( auto VertexDec : Shader.VertexDeclaration )
-		{
-			if ( VertexDec.Stride > 0 )
-			{
-				RHICmdList.SetStreamSource(VertexDec.StreamIndex, PredrawVB, VertexDec.Stride, VertexDec.Offset);
-			}
-			else
-			{
-				RHICmdList.SetStreamSource(VertexDec.StreamIndex, PredrawZVB, VertexDec.Stride, VertexDec.Offset);
-			}
-		}
-		
 		FShaderPlatformCache& PlatformCache = Caches.FindOrAdd(GMaxRHIShaderPlatform);
 		for ( auto DrawKeyIdx : DrawStates )
 		{
@@ -1233,10 +1347,6 @@ void FShaderCache::PreDrawShader(FRHICommandList& RHICmdList, FShaderCacheBoundS
 			FDepthStencilStateRHIRef DepthStencil = RHICreateDepthStencilState(DrawKey.DepthStencilState);
 			FRasterizerStateRHIRef Rasterizer = RHICreateRasterizerState(DrawKey.RasterizerState);
 			
-			RHICmdList.SetBlendState(BlendState);
-			RHICmdList.SetDepthStencilState(DepthStencil);
-			RHICmdList.SetRasterizerState(Rasterizer);
-			
 			uint32 NewNumRenderTargets = 0;
 			FRHIRenderTargetView RenderTargets[MaxSimultaneousRenderTargets];
 			for ( uint32 i = 0; i < MaxSimultaneousRenderTargets; i++ )
@@ -1271,6 +1381,22 @@ void FShaderCache::PreDrawShader(FRHICommandList& RHICmdList, FShaderCacheBoundS
 			
 			RHICmdList.SetRenderTargets(NewNumRenderTargets, RenderTargets, bDepthStencilTarget ? &DepthStencilTarget : nullptr, 0, nullptr);
 			
+			RHICmdList.SetBlendState(BlendState);
+			RHICmdList.SetDepthStencilState(DepthStencil);
+			RHICmdList.SetRasterizerState(Rasterizer);
+
+			for( auto VertexDec : Shader.VertexDeclaration )
+			{
+				if ( VertexDec.Stride > 0 )
+				{
+					RHICmdList.SetStreamSource(VertexDec.StreamIndex, PredrawVB, VertexDec.Stride, VertexDec.Offset);
+				}
+				else
+				{
+					RHICmdList.SetStreamSource(VertexDec.StreamIndex, PredrawZVB, VertexDec.Stride, VertexDec.Offset);
+				}
+			}
+
 			if ( !IsValidRef(ShaderBoundState) )
 			{
 				FVertexShaderRHIRef VertexShader = Shader.VertexShader.bActive ? CachedVertexShaders.FindRef(Shader.VertexShader) : nullptr;
diff --git a/Engine/Source/Runtime/ShaderCore/Public/ShaderCache.h b/Engine/Source/Runtime/ShaderCore/Public/ShaderCache.h
--- a/Engine/Source/Runtime/ShaderCore/Public/ShaderCache.h
+++ b/Engine/Source/Runtime/ShaderCore/Public/ShaderCache.h
@@ -15,7 +15,13 @@ struct SHADERCORE_API FShaderCacheCustomVersion
 {
 	static const FGuid Key;
 	static const FGuid GameKey;
-	enum Type {	Initial, PreDraw, CacheHashes, OptimisedHashes, StreamingKeys, Latest = StreamingKeys };
+	enum Type {	Initial, PreDraw, CacheHashes, OptimisedHashes, StreamingKeys, AdditionalResources, SeparateBinaries, Latest = SeparateBinaries };
+};
+
+enum EShaderCacheOptions
+{
+	SCO_Default,
+	SCO_NoShaderPreload = 1 << 0 /* Disable preloading of shaders for RHIs where loading all shaders is too slow (i.e. Metal online compiler). */
 };
 
 /** Texture type enum for shader cache draw keys */
@@ -242,7 +248,8 @@ class SHADERCORE_API FShaderCache
 	{
 		enum
 		{
-			MaxTextureSamplers = 16,
+			MaxNumSamplers = 16,
+			MaxNumResources = 128,
 			NullState = ~0u
 		};
 		
@@ -257,7 +264,7 @@ class SHADERCORE_API FShaderCache
 			FMemory::Memset(RenderTargets, 255, sizeof(RenderTargets));
 			FMemory::Memset(SamplerStates, 255, sizeof(SamplerStates));
 			FMemory::Memset(Resources, 255, sizeof(Resources));
-			check(GetFeatureLevelMaxTextureSamplers(GMaxRHIFeatureLevel) <= MaxTextureSamplers);
+			check(GetFeatureLevelMaxTextureSamplers(GMaxRHIFeatureLevel) <= MaxNumSamplers);
 		}
 		
 		struct FShaderRasterizerState
@@ -315,8 +322,8 @@ class SHADERCORE_API FShaderCache
 		FShaderRasterizerState RasterizerState;
 		FDepthStencilStateInitializerRHI DepthStencilState;
 		uint32 RenderTargets[MaxSimultaneousRenderTargets];
-		uint32 SamplerStates[SF_NumFrequencies][MaxTextureSamplers];
-		uint32 Resources[SF_NumFrequencies][MaxTextureSamplers];
+		uint32 SamplerStates[SF_NumFrequencies][MaxNumSamplers];
+		uint32 Resources[SF_NumFrequencies][MaxNumResources];
 		uint32 DepthStencilTarget;
 		mutable uint32 Hash;
 		uint8 IndexType;
@@ -353,9 +360,12 @@ class SHADERCORE_API FShaderCache
 				
 				for( uint32 i = 0; i < SF_NumFrequencies; i++ )
 				{
-					for( uint32 j = 0; j < MaxTextureSamplers; j++ )
+					for( uint32 j = 0; j < MaxNumSamplers; j++ )
 					{
 						Key.Hash ^= Key.SamplerStates[i][j];
+					}
+					for( uint32 j = 0; j < MaxNumResources; j++ )
+					{
 						Key.Hash ^= Key.Resources[i][j];
 					}
 				}
@@ -387,9 +397,12 @@ class SHADERCORE_API FShaderCache
 		{
 			for ( uint32 i = 0; i < SF_NumFrequencies; i++ )
 			{
-				for ( uint32 j = 0; j < MaxTextureSamplers; j++ )
+				for ( uint32 j = 0; j < MaxNumSamplers; j++ )
 				{
 					Ar << Info.SamplerStates[i][j];
+				}
+				for ( uint32 j = 0; j < MaxNumResources; j++ )
+				{
 					Ar << Info.Resources[i][j];
 				}
 			}
@@ -407,9 +420,10 @@ public:
 	
 	/** Called by the game to set the game specific shader cache version, only caches of this version will be loaded. Must be called before RHI initialisation, as InitShaderCache will load any existing cache. */
 	static void SetGameVersion(int32 InGameVersion);
+	static int32 GetGameVersion() { return GameVersion; }
 	
 	/** Shader cache initialisation, called only by the RHI. */
-	static void InitShaderCache();
+	static void InitShaderCache(EShaderCacheOptions Options, uint32 InMaxResources = FShaderDrawKey::MaxNumResources);
 	/** Shader cache shutdown, called only by the RHI. */
 	static void ShutdownShaderCache();
 	
@@ -656,9 +670,26 @@ public:
 		}
 	}
 	
+	static FORCEINLINE bool IsPredrawCall()
+	{
+		FShaderCache* ShaderCache = GetShaderCache();
+		if ( Cache )
+		{
+			return Cache->bIsPreDraw;
+		}
+		return false;
+	}
+
 	/** Archive serialisation of the cache data. */
 	friend FArchive& operator<<( FArchive& Ar, FShaderCache& Info );
 	
+	struct SHADERCORE_API FShaderCodeCache
+	{
+		friend FArchive& operator<<( FArchive& Ar, FShaderCodeCache& Info );
+		
+		TMap<FShaderCacheKey, TArray<uint8>> Shaders;
+	};
+	
 private:
 	struct FShaderCacheBoundState
 	{
@@ -863,6 +894,9 @@ private:
 	// Serialised
 	FShaderCaches Caches;
 	
+	// Optional, separate code cache
+	FShaderCodeCache CodeCache;
+	
 	// Transient non-invasive tracking of RHI resources for shader logging
 	TMap<FShaderCacheKey, FVertexShaderRHIRef> CachedVertexShaders;
 	TMap<FShaderCacheKey, FPixelShaderRHIRef> CachedPixelShaders;
@@ -921,6 +955,8 @@ private:
 	uint32 Viewport[4];
 	float DepthRange[2];
 	bool bIsPreDraw;
+	EShaderCacheOptions Options;
+	uint8 MaxResources;
 	
 	static FShaderCache* Cache;
 	static int32 GameVersion;
@@ -928,8 +964,10 @@ private:
 	static int32 bUseShaderPredraw;
 	static int32 bUseShaderDrawLog;
 	static int32 PredrawBatchTime;
+	static int32 bUseShaderBinaryCache;
 	static FAutoConsoleVariableRef CVarUseShaderCaching;
 	static FAutoConsoleVariableRef CVarUseShaderPredraw;
 	static FAutoConsoleVariableRef CVarUseShaderDrawLog;
 	static FAutoConsoleVariableRef CVarPredrawBatchTime;
+	static FAutoConsoleVariableRef CVarUseShaderBinaryCache;
 };
