Batched shader precompilation & predraw for Mac OS X which can be more aggressive behind loading screens & purged before entering an actual level, to reduce the time taken to reach an interactive state without hitches in the actual game.
- With 'r.UseAsyncShaderPrecompilation' enabled shader precompilation is now executed on the render thread alongside the predraw, not as a single operation on loading the binary cache.
- The setting 'r.TargetPrecompileFrameTime' controls what the target frame rate should be when performing async. precompile and the shader cache will try and adjust how many shaders to precompile to stay within this limit, with a minimum of 1 shader per frame.
- The shader cache is now updated during initial loading to take advantage of the time spent loading essential data to do as much of the precompile as possible.
- It is possible to accelerate the shader cache's precompilation by specifying 'r.AccelPredrawBatchTime' and 'r.AccelTargetPrecompileFrameTime' to larger values, then calling FShaderCache::BeginAcceleratedBatching() to start using the more aggressive values and FShaderCache::EndAcceleratedBatching() to go back to the normal values.
- Calling FShaderCache::FlushOutstandingBatches() will force completion of all outstanding batched precompilation and predraws during the next frame.

diff --git a/Engine/Source/Runtime/Launch/Private/LaunchEngineLoop.cpp b/Engine/Source/Runtime/Launch/Private/LaunchEngineLoop.cpp
--- a/Engine/Source/Runtime/Launch/Private/LaunchEngineLoop.cpp
+++ b/Engine/Source/Runtime/Launch/Private/LaunchEngineLoop.cpp
@@ -36,7 +36,6 @@
 	#include "DerivedDataCacheInterface.h"
 	#include "RenderCore.h"
 	#include "ShaderCompiler.h"
-	#include "ShaderCache.h"
 	#include "DistanceFieldAtlas.h"
 	#include "GlobalShader.h"
 	#include "ParticleHelper.h"
@@ -2722,7 +2721,6 @@ void FEngineLoop::Tick()
 		ENQUEUE_UNIQUE_RENDER_COMMAND(
 			EndFrame,
 		{
-			FShaderCache::PreDrawShaders(RHICmdList);
 			RHICmdList.EndFrame();
 			RHICmdList.PopEvent();
 		});
diff --git a/Engine/Source/Runtime/ShaderCore/Private/ShaderCache.cpp b/Engine/Source/Runtime/ShaderCore/Private/ShaderCache.cpp
--- a/Engine/Source/Runtime/ShaderCore/Private/ShaderCache.cpp
+++ b/Engine/Source/Runtime/ShaderCore/Private/ShaderCache.cpp
@@ -14,6 +14,13 @@ DECLARE_STATS_GROUP(TEXT("Shader Cache"),STATGROUP_ShaderCache, STATCAT_Advanced
 DECLARE_DWORD_ACCUMULATOR_STAT(TEXT("Num Shaders Cached"),STATGROUP_NumShadersCached,STATGROUP_ShaderCache);
 DECLARE_DWORD_ACCUMULATOR_STAT(TEXT("Num BSS Cached"),STATGROUP_NumBSSCached,STATGROUP_ShaderCache);
 DECLARE_DWORD_ACCUMULATOR_STAT(TEXT("Num New Draw-States Cached"),STATGROUP_NumDrawsCached,STATGROUP_ShaderCache);
+DECLARE_DWORD_COUNTER_STAT(TEXT("Shaders Precompiled"),STATGROUP_NumPrecompiled,STATGROUP_ShaderCache);
+DECLARE_DWORD_COUNTER_STAT(TEXT("Shaders Predrawn"),STATGROUP_NumPredrawn,STATGROUP_ShaderCache);
+DECLARE_DWORD_COUNTER_STAT(TEXT("Draw States Predrawn"),STATGROUP_NumStatesPredrawn,STATGROUP_ShaderCache);
+DECLARE_DWORD_ACCUMULATOR_STAT(TEXT("Total Shaders Precompiled"),STATGROUP_TotalPrecompiled,STATGROUP_ShaderCache);
+DECLARE_DWORD_ACCUMULATOR_STAT(TEXT("Total Shaders Predrawn"),STATGROUP_TotalPredrawn,STATGROUP_ShaderCache);
+DECLARE_DWORD_ACCUMULATOR_STAT(TEXT("Total Draw States Predrawn"),STATGROUP_TotalStatesPredrawn,STATGROUP_ShaderCache);
+DECLARE_DWORD_COUNTER_STAT(TEXT("Num To Precompile Per Frame"),STATGROUP_NumToPrecompile,STATGROUP_ShaderCache);
 
 const FGuid FShaderCacheCustomVersion::Key(0xB954F018, 0xC9624DD6, 0xA74E79B1, 0x8EA113C2);
 const FGuid FShaderCacheCustomVersion::GameKey(0x03D4EB48, 0xB50B4CC3, 0xA598DE41, 0x5C6CC993);
@@ -72,6 +79,40 @@ FAutoConsoleVariableRef FShaderCache::CVarUseShaderBinaryCache(
 	ECVF_ReadOnly|ECVF_RenderThreadSafe
 	);
 
+// Whether to try and perform shader precompilation asynchronously.
+int32 FShaderCache::bUseAsyncShaderPrecompilation = 0;
+FAutoConsoleVariableRef FShaderCache::CVarUseAsyncShaderPrecompilation(
+	TEXT("r.UseAsyncShaderPrecompilation"),
+	bUseAsyncShaderPrecompilation,
+	TEXT("If true tries to perform inital shader precompilation asynchronously on a background thread. Defaults to false."),
+	ECVF_ReadOnly|ECVF_RenderThreadSafe
+	);
+
+// As async precompile can take significant time specify a desired max. frame time that the cache will try to remain below while precompiling. We can't specify the time to spend directly as under GL compile operations are deferred and take no time on the user thread.
+int32 FShaderCache::TargetPrecompileFrameTime = -1;
+FAutoConsoleVariableRef FShaderCache::CVarTargetPrecompileFrameTime(
+	TEXT("r.TargetPrecompileFrameTime"),
+	TargetPrecompileFrameTime,
+	TEXT("Upper limit in ms for total frame time while precompiling, allowing the shader cache to adjust how many shaders to precompile each frame. Defaults to -1 which will precompile all shaders immediately."),
+	ECVF_RenderThreadSafe
+	);
+
+int32 FShaderCache::AccelPredrawBatchTime = 0;
+FAutoConsoleVariableRef FShaderCache::CVarAccelPredrawBatchTime(
+	TEXT("r.AccelPredrawBatchTime"),
+	AccelPredrawBatchTime,
+	TEXT("Override value for r.PredrawBatchTime when showing a loading-screen or similar to do more work while the player won't notice, or 0 to use r.PredrawBatchTime. Defaults to 0."),
+	ECVF_RenderThreadSafe
+	);
+
+int32 FShaderCache::AccelTargetPrecompileFrameTime = 0;
+FAutoConsoleVariableRef FShaderCache::CVarAccelTargetPrecompileFrameTime(
+	TEXT("r.AccelTargetPrecompileFrameTime"),
+	AccelTargetPrecompileFrameTime,
+	TEXT("Override value for r.TargetPrecompileFrameTime when showing a loading-screen or similar to do more work while the player won't notice, or 0 to use r.TargetPrecompileFrameTime. Defaults to 0."),
+	ECVF_RenderThreadSafe
+	);
+
 FShaderCache* FShaderCache::Cache = nullptr;
 int32 FShaderCache::GameVersion = 0;
 
@@ -161,13 +202,46 @@ void FShaderCache::LoadBinaryCache()
 			}
 		}
 		
-		if(bLoadedCache)
+		if ( bLoadedCache )
 		{
-			for(auto Entry : Cache->CodeCache.Shaders)
+			if ( !bUseAsyncShaderPrecompilation || FParse::Param(FCommandLine::Get(), TEXT("DisableAsyncShaderPrecompilation")) )
 			{
-				if(Entry.Key.Platform == GMaxRHIShaderPlatform)
+				for ( auto Entry : Cache->CodeCache.Shaders )
 				{
-					Cache->InternalLogShader(Entry.Key.Platform, Entry.Key.Frequency, Entry.Key.SHAHash, Entry.Value);
+					if ( Entry.Key.Platform == GMaxRHIShaderPlatform )
+					{
+						Cache->InternalLogShader(Entry.Key.Platform, Entry.Key.Frequency, Entry.Key.SHAHash, Entry.Value);
+					}
+				}
+			}
+			else
+			{
+				for ( auto Entry : Cache->CodeCache.Shaders )
+				{
+					bool bUsable = FShaderResource::ArePlatformsCompatible(GMaxRHIShaderPlatform, Entry.Key.Platform);
+					switch (Entry.Key.Frequency)
+					{
+						case SF_Geometry:
+							bUsable &= RHISupportsGeometryShaders(GMaxRHIShaderPlatform);
+							break;
+							
+						case SF_Hull:
+						case SF_Domain:
+							bUsable &= RHISupportsTessellation(GMaxRHIShaderPlatform);
+							break;
+							
+						case SF_Compute:
+							bUsable &= RHISupportsComputeShaders(GMaxRHIShaderPlatform);
+							break;
+							
+						default:
+							break;
+					}
+					
+					if ( bUsable )
+					{
+						Cache->ShadersToPrecompile.Emplace(Entry.Key);
+					}
 				}
 			}
 		}
@@ -228,6 +302,8 @@ FShaderCache::FShaderCache()
 , CurrentShaderState(nullptr)
 , bIsPreDraw(false)
 , Options(SCO_Default)
+, OverridePrecompileTime(0)
+, OverridePredrawBatchTime(0)
 {
 	Viewport[0] = Viewport[1] = Viewport[2] = Viewport[3] = 0;
 	DepthRange[0] = DepthRange[1] = 0.0f;
@@ -481,18 +557,19 @@ void FShaderCache::InternalLogShader(EShaderPlatform Platform, EShaderFrequency
 		Key.Platform = Platform;
 		Key.Frequency = Frequency;
 		Key.bActive = true;
-		
-		if(bUseShaderBinaryCache && !CodeCache.Shaders.Contains(Key))
-		{
-			CodeCache.Shaders.Add(Key, Code);
-		}
-		
-		if(!(Options & SCO_NoShaderPreload))
-		{
-			ENQUEUE_UNIQUE_RENDER_COMMAND_TWOPARAMETER(LogShader, FShaderCacheKey,Key,Key, TArray<uint8>,Code,Code, {
-				FShaderCache::GetShaderCache()->SubmitShader(Key, Code);
-			});
-		}
+
+		ENQUEUE_UNIQUE_RENDER_COMMAND_THREEPARAMETER(LogShader, FShaderCache*,ShaderCache,this, FShaderCacheKey,Key,Key, TArray<uint8>,Code,Code, {
+			bool bSubmit = !ShaderCache->bUseShaderBinaryCache || !ShaderCache->bUseAsyncShaderPrecompilation;
+			if(ShaderCache->bUseShaderBinaryCache && !ShaderCache->CodeCache.Shaders.Contains(Key))
+			{
+				ShaderCache->CodeCache.Shaders.Add(Key, Code);
+				bSubmit |= true;
+			}
+			if(!(ShaderCache->Options & SCO_NoShaderPreload) && bSubmit)
+			{
+				ShaderCache->SubmitShader(Key, Code);
+			}
+		});
 	}
 }
 
@@ -897,7 +974,8 @@ void FShaderCache::InternalLogDraw(uint8 IndexType)
 			Id = PlatformCache.DrawStates.Add(CurrentDrawKey);
 		}
 		
-		TSet<int32>& ShaderDrawSet = PlatformCache.StreamingDrawStates.FindOrAdd(StreamingKey).ShaderDrawStates.FindOrAdd(BoundShaderState);
+		FShaderStreamingCache& StreamCache = PlatformCache.StreamingDrawStates.FindOrAdd(StreamingKey);
+		TSet<int32>& ShaderDrawSet = StreamCache.ShaderDrawStates.FindOrAdd(BoundShaderState);
 		if( !ShaderDrawSet.Contains(Id.AsInteger()) )
 		{
 			INC_DWORD_STAT(STATGROUP_NumDrawsCached);
@@ -907,88 +985,208 @@ void FShaderCache::InternalLogDraw(uint8 IndexType)
 		// No need to predraw this shader draw key - we've already done it
 		for(auto StreamingMap : ShadersToDraw)
 		{
-			StreamingMap.Value.ShaderDrawStates.FindRef(BoundShaderState).Remove(Id.AsInteger());
+			TSet<int32>* Set = StreamingMap.Value.ShaderDrawStates.Find(BoundShaderState);
+			if ( Set )
+			{
+				Set->Remove(Id.AsInteger());
+			}
 		}
 	}
 }
 
-void FShaderCache::InternalPreDrawShaders(FRHICommandList& RHICmdList)
+void FShaderCache::InternalPreDrawShaders(FRHICommandList& RHICmdList, float DeltaTime)
 {
-	if ( bUseShaderPredraw && ShadersToDraw.FindRef(StreamingKey).ShaderDrawStates.Num() > 0 )
+	static bool FirstCall = true;
+	static uint32 NumShadersToCompile = 1;
+	
+	static uint32 FrameNum = 0;
+	
+	if(FrameNum != GFrameNumberRenderThread || OverridePrecompileTime != 0 || OverridePredrawBatchTime != 0)
 	{
-		bIsPreDraw = true;
-		
-		if ( !IsValidRef(IndexBufferUInt16) )
+		FrameNum = GFrameNumberRenderThread;
+		uint32 NumCompiled = 0;
+		int64 TimeForPredrawing = 0;
+		if ( bUseShaderBinaryCache && bUseAsyncShaderPrecompilation && ShadersToPrecompile.Num() > 0 )
 		{
-			FRHIResourceCreateInfo Info;
-			uint32 Stride = sizeof(uint16);
-			uint32 Size = sizeof(uint16) * 3;
-			void* Data = nullptr;
-			IndexBufferUInt16 = RHICreateAndLockIndexBuffer(Stride, Size, BUF_Static, Info, Data);			
-			if ( Data )
+			SET_DWORD_STAT(STATGROUP_NumToPrecompile, NumShadersToCompile);
+			
+			for( uint32 Index = 0; (GetTargetPrecompileFrameTime() == -1 || NumCompiled < NumShadersToCompile) && Index < ShadersToPrecompile.Num(); ++Index )
 			{
-				FMemory::Memzero(Data, Size);
+				FShaderCacheKey& Key = ShadersToPrecompile[Index];
+				TArray<uint8>& Code = CodeCache.Shaders[Key];
+				SubmitShader(Key, Code);
+				INC_DWORD_STAT(STATGROUP_NumPrecompiled);
+				INC_DWORD_STAT(STATGROUP_TotalPrecompiled);
+				
+				uint32 OldIndex = Index;
+				--Index;
+				ShadersToPrecompile.RemoveAt(OldIndex);
+				
+				++NumCompiled;
 			}
-			RHIUnlockIndexBuffer(IndexBufferUInt16);
-		}
-		if ( !IsValidRef(IndexBufferUInt32) )
-		{
-			FRHIResourceCreateInfo Info;
-			uint32 Stride = sizeof(uint32);
-			uint32 Size = sizeof(uint32) * 3;
-			void* Data = nullptr;
-			IndexBufferUInt32 = RHICreateAndLockIndexBuffer(Stride, Size, BUF_Static, Info, Data);			
-			if ( Data )
+			
+			if(GetTargetPrecompileFrameTime() != -1)
 			{
-				FMemory::Memzero(Data, Size);
+				int64 MSec = DeltaTime * 1000.0;
+				if( MSec < GetTargetPrecompileFrameTime() )
+				{
+					NumShadersToCompile++;
+				}
+				else
+				{
+					NumShadersToCompile = FMath::Max(1u, NumShadersToCompile / 2u);
+				}
+				
+				if(GetPredrawBatchTime() != -1)
+				{
+					TimeForPredrawing += FMath::Max((MSec - (int64)GetTargetPrecompileFrameTime()), (int64)0);
+				}
 			}
-			RHIUnlockIndexBuffer(IndexBufferUInt32);
 		}
 		
-		RHICmdList.SetViewport(0, 0, FLT_MIN, 3, 3, FLT_MAX);
-		
-		int64 TimeForPredrawing = 0;
-		TMap<FShaderCacheBoundState, TSet<int32>>& ShaderDrawStates = ShadersToDraw.FindOrAdd(StreamingKey).ShaderDrawStates;
-		for ( auto It = ShaderDrawStates.CreateIterator(); (PredrawBatchTime == -1 || TimeForPredrawing < PredrawBatchTime) && It; ++It )
+		if ( bUseShaderPredraw && !FirstCall && (GetPredrawBatchTime() == -1 || TimeForPredrawing < GetPredrawBatchTime()) && ShadersToDraw.FindRef(StreamingKey).ShaderDrawStates.Num() > 0 )
 		{
-			uint32 Start = FPlatformTime::Cycles();
+			bIsPreDraw = true;
+			
+			if ( !IsValidRef(IndexBufferUInt16) )
+			{
+				FRHIResourceCreateInfo Info;
+				uint32 Stride = sizeof(uint16);
+				uint32 Size = sizeof(uint16) * 3;
+				void* Data = nullptr;
+				IndexBufferUInt16 = RHICreateAndLockIndexBuffer(Stride, Size, BUF_Static, Info, Data);			
+				if ( Data )
+				{
+					FMemory::Memzero(Data, Size);
+				}
+				RHIUnlockIndexBuffer(IndexBufferUInt16);
+			}
+			if ( !IsValidRef(IndexBufferUInt32) )
+			{
+				FRHIResourceCreateInfo Info;
+				uint32 Stride = sizeof(uint32);
+				uint32 Size = sizeof(uint32) * 3;
+				void* Data = nullptr;
+				IndexBufferUInt32 = RHICreateAndLockIndexBuffer(Stride, Size, BUF_Static, Info, Data);			
+				if ( Data )
+				{
+					FMemory::Memzero(Data, Size);
+				}
+				RHIUnlockIndexBuffer(IndexBufferUInt32);
+			}
 			
-			auto Shader = *It;
-			TSet<int32>& ShaderDrawSet = Shader.Value;
-			PreDrawShader(RHICmdList, Shader.Key, ShaderDrawSet);
-			It.RemoveCurrent();
+			RHICmdList.SetViewport(0, 0, FLT_MIN, 3, 3, FLT_MAX);
 			
-			uint32 End = FPlatformTime::Cycles();
-			TimeForPredrawing += FPlatformTime::ToMilliseconds(End - Start);
+			TMap<FShaderCacheBoundState, TSet<int32>>& ShaderDrawStates = ShadersToDraw.FindOrAdd(StreamingKey).ShaderDrawStates;
+			for ( auto It = ShaderDrawStates.CreateIterator(); (GetPredrawBatchTime() == -1 || TimeForPredrawing < GetPredrawBatchTime()) && It; ++It )
+			{
+				uint32 Start = FPlatformTime::Cycles();
+				
+				auto Shader = *It;
+				TSet<int32>& ShaderDrawSet = Shader.Value;
+				PreDrawShader(RHICmdList, Shader.Key, ShaderDrawSet);
+				It.RemoveCurrent();
+				
+				uint32 End = FPlatformTime::Cycles();
+				TimeForPredrawing += FPlatformTime::ToMilliseconds(End - Start);
+			}
+			
+			// This is a bit dirty/naughty but it forces the draw commands to be flushed through on OS X
+			// which means we can delete the resources without crashing MTGL.
+			RHIFlushResources();
+			
+			RHICmdList.SetBoundShaderState(CurrentShaderState);
+			
+			FBlendStateRHIRef BlendState = RHICreateBlendState(CurrentDrawKey.BlendState);
+			FDepthStencilStateRHIRef DepthStencil = RHICreateDepthStencilState(CurrentDrawKey.DepthStencilState);
+			FRasterizerStateRHIRef Rasterizer = RHICreateRasterizerState(CurrentDrawKey.RasterizerState);
+			
+			RHICmdList.SetBlendState(BlendState);
+			RHICmdList.SetDepthStencilState(DepthStencil);
+			RHICmdList.SetRasterizerState(Rasterizer);
+
+			RHICmdList.SetViewport(Viewport[0], Viewport[1], DepthRange[0], Viewport[2], Viewport[3], DepthRange[1]);
+			
+			if ( ShadersToDraw.FindOrAdd(StreamingKey).ShaderDrawStates.Num() == 0 )
+			{
+				PredrawRTs.Empty();
+				PredrawBindings.Empty();
+				PredrawVBs.Empty();
+			}
+			
+			bIsPreDraw = false;
 		}
 		
-		// This is a bit dirty/naughty but it forces the draw commands to be flushed through on OS X
-		// which means we can delete the resources without crashing MTGL.
-		RHIFlushResources();
-		
-		RHICmdList.SetBoundShaderState(CurrentShaderState);
+		if(OverridePrecompileTime == -1)
+		{
+			OverridePrecompileTime = 0;
+		}
 		
-		FBlendStateRHIRef BlendState = RHICreateBlendState(CurrentDrawKey.BlendState);
-		FDepthStencilStateRHIRef DepthStencil = RHICreateDepthStencilState(CurrentDrawKey.DepthStencilState);
-		FRasterizerStateRHIRef Rasterizer = RHICreateRasterizerState(CurrentDrawKey.RasterizerState);
+		if(OverridePredrawBatchTime == -1)
+		{
+			OverridePredrawBatchTime = 0;
+		}
 		
-		RHICmdList.SetBlendState(BlendState);
-		RHICmdList.SetDepthStencilState(DepthStencil);
-		RHICmdList.SetRasterizerState(Rasterizer);
+		FirstCall = false;
+	}
+}
 
-		RHICmdList.SetViewport(Viewport[0], Viewport[1], DepthRange[0], Viewport[2], Viewport[3], DepthRange[1]);
-		
-		if ( ShadersToDraw.FindOrAdd(StreamingKey).ShaderDrawStates.Num() == 0 )
+void FShaderCache::BeginAcceleratedBatching()
+{
+	if ( Cache )
+	{
+		if(AccelTargetPrecompileFrameTime)
 		{
-			PredrawRTs.Empty();
-			PredrawBindings.Empty();
-			PredrawVBs.Empty();
+			Cache->OverridePrecompileTime = AccelTargetPrecompileFrameTime;
 		}
-		
-		bIsPreDraw = false;
+		if(AccelPredrawBatchTime)
+		{
+			Cache->OverridePredrawBatchTime = AccelPredrawBatchTime;
+		}
+	}
+}
+
+void FShaderCache::EndAcceleratedBatching()
+{
+	if ( Cache )
+	{
+		Cache->OverridePrecompileTime = 0;
+		Cache->OverridePredrawBatchTime = 0;
+	}
+}
+
+void FShaderCache::FlushOutstandingBatches()
+{
+	if ( Cache )
+	{
+		Cache->OverridePrecompileTime = -1;
+		Cache->OverridePredrawBatchTime = -1;
+	}
+}
+
+void FShaderCache::Tick( float DeltaTime )
+{
+	if ( Cache )
+	{
+		Cache->InternalPreDrawShaders(GRHICommandList.GetImmediateCommandList(), DeltaTime);
 	}
 }
 
+bool FShaderCache::IsTickable() const
+{
+	return ( bUseShaderBinaryCache && bUseAsyncShaderPrecompilation && ShadersToPrecompile.Num() > 0 ) || ( bUseShaderPredraw && ShadersToDraw.FindRef(StreamingKey).ShaderDrawStates.Num() > 0 );
+}
+
+bool FShaderCache::NeedsRenderingResumedForRenderingThreadTick() const
+{
+	return true;
+}
+
+TStatId FShaderCache::GetStatId() const
+{
+	RETURN_QUICK_DECLARE_CYCLE_STAT(FShaderCache, STATGROUP_Tickables);
+}
+
 void FShaderCache::PrebindShader(FShaderCacheKey const& Key)
 {
 	FShaderPlatformCache& PlatformCache = Caches.FindOrAdd(Key.Platform);
@@ -1502,6 +1700,8 @@ void FShaderCache::PreDrawShader(FRHICommandList& RHICmdList, FShaderCacheBoundS
 					break;
 				}
 			}
+			INC_DWORD_STAT(STATGROUP_NumStatesPredrawn);
+			INC_DWORD_STAT(STATGROUP_TotalStatesPredrawn);
 		}
 		
 		if( IsValidRef( ShaderBoundState ) && DrawStates.Num() )
@@ -1532,6 +1732,9 @@ void FShaderCache::PreDrawShader(FRHICommandList& RHICmdList, FShaderCacheBoundS
 		{
 			RHICmdList.SetStreamSource(VertexDec.StreamIndex, nullptr, 0, 0);
 		}
+		
+		INC_DWORD_STAT(STATGROUP_NumPredrawn);
+		INC_DWORD_STAT(STATGROUP_TotalPredrawn);
 	}
 }
 
@@ -1558,3 +1761,13 @@ uint32 FShaderCache::FSamplerStateInitializerRHIKeyFuncs::GetKeyHash(KeyInitType
 {
 	return FCrc::MemCrc_DEPRECATED(&Key, CalculateSizeOfSamplerStateInitializer());
 }
+
+int32 FShaderCache::GetPredrawBatchTime() const
+{
+	return OverridePrecompileTime == 0 ? PredrawBatchTime : OverridePrecompileTime;
+}
+
+int32 FShaderCache::GetTargetPrecompileFrameTime() const
+{
+	return OverridePredrawBatchTime == 0 ? TargetPrecompileFrameTime : OverridePredrawBatchTime;
+}
diff --git a/Engine/Source/Runtime/ShaderCore/Public/ShaderCache.h b/Engine/Source/Runtime/ShaderCore/Public/ShaderCache.h
--- a/Engine/Source/Runtime/ShaderCore/Public/ShaderCache.h
+++ b/Engine/Source/Runtime/ShaderCore/Public/ShaderCache.h
@@ -9,6 +9,7 @@
 #include "ShaderCore.h"
 #include "RHI.h"
 #include "BoundShaderStateCache.h"
+#include "TickableObjectRenderThread.h"
 
 /** Custom serialization version for FShaderCache */
 struct SHADERCORE_API FShaderCacheCustomVersion
@@ -205,8 +206,10 @@ struct SHADERCORE_API FShaderRenderTargetKey
  * Supported RHIs:
  * - OpenGLDrv
  */
-class SHADERCORE_API FShaderCache
+class SHADERCORE_API FShaderCache : public FTickableObjectRenderThread
 {
+	friend class FShaderCacheLoadBinaryCodeTask;
+	
 	struct SHADERCORE_API FShaderCacheKey
 	{
 		FShaderCacheKey() : Platform(SP_NumPlatforms), Frequency(SF_NumFrequencies), Hash(0), bActive(false) {}
@@ -416,7 +419,7 @@ class SHADERCORE_API FShaderCache
 	
 public:
 	FShaderCache();
-	~FShaderCache();
+	virtual ~FShaderCache();
 	
 	/** Called by the game to set the game specific shader cache version, only caches of this version will be loaded. Must be called before RHI initialisation, as InitShaderCache will load any existing cache. */
 	static void SetGameVersion(int32 InGameVersion);
@@ -663,15 +666,7 @@ public:
 		}
 	}
 	
-	/** Called by the RHI. Predraws a batch of shaders if and only if there are any outstanding from the current shader cache and r.UseShaderCaching & r.UseShaderPredraw are enabled. */
-	static FORCEINLINE void PreDrawShaders(FRHICommandList& RHICmdList)
-	{
-		if ( Cache )
-		{
-			Cache->InternalPreDrawShaders(RHICmdList);
-		}
-	}
-	
+	/** Called by the RHI. Returns whether the current draw call is a predraw call for shader variant submission in the underlying driver rather than a real UE4 draw call. */
 	static FORCEINLINE bool IsPredrawCall()
 	{
 		FShaderCache* ShaderCache = GetShaderCache();
@@ -681,6 +676,15 @@ public:
 		}
 		return false;
 	}
+	
+	/** Use the accelerated target precompile frame time (in ms) and predraw batch time (in ms) during a loading screen or other scenario where frame time can be much slower than normal. */
+	static void BeginAcceleratedBatching();
+	
+	/** Stops applying the overrides applied with BeginAcceleratedBatching but doesn't flush outstanding work. */
+	static void EndAcceleratedBatching();
+	
+	/** Flushes all outstanding precompilation/predraw work & then ends accelerated batching as per EndAcceleratedBatching. */
+	static void FlushOutstandingBatches();
 
 	/** Archive serialisation of the cache data. */
 	friend FArchive& operator<<( FArchive& Ar, FShaderCache& Info );
@@ -692,6 +696,15 @@ public:
 		TMap<FShaderCacheKey, TArray<uint8>> Shaders;
 	};
 	
+public: // From FTickableObjectRenderThread
+	virtual void Tick( float DeltaTime ) final override;
+	
+	virtual bool IsTickable() const final override;
+	
+	virtual bool NeedsRenderingResumedForRenderingThreadTick() const final override;
+	
+	virtual TStatId GetStatId() const final override;
+	
 private:
 	struct FShaderCacheBoundState
 	{
@@ -881,7 +894,7 @@ private:
 	void InternalSetViewport(uint32 MinX, uint32 MinY, float MinZ, uint32 MaxX, uint32 MaxY, float MaxZ);
 	
 	void InternalLogDraw(uint8 IndexType);
-	void InternalPreDrawShaders(FRHICommandList& RHICmdList);
+	void InternalPreDrawShaders(FRHICommandList& RHICmdList, float DeltaTime);
 	
 	void PrebindShader(FShaderCacheKey const& Key);
 	void SubmitShader(FShaderCacheKey const& Key, TArray<uint8> const& Code);
@@ -892,6 +905,9 @@ private:
 	FShaderTextureBinding CreateSRV(FShaderResourceKey const& ResourceKey);
 	FTextureRHIRef CreateRenderTarget(FShaderRenderTargetKey const& TargetKey);
 	
+	int32 GetPredrawBatchTime() const;
+	int32 GetTargetPrecompileFrameTime() const;
+	
 private:
 	// Serialised
 	FShaderCaches Caches;
@@ -925,6 +941,9 @@ private:
 	// Logged draws will only be predrawn when this key becomes active.
 	uint32 StreamingKey;
 	
+	// Shaders to precompile
+	TArray<FShaderCacheKey> ShadersToPrecompile;
+	
 	// Shaders we need to predraw
 	TMap<uint32, FShaderStreamingCache> ShadersToDraw;
 	
@@ -960,6 +979,10 @@ private:
 	EShaderCacheOptions Options;
 	uint8 MaxResources;
 	
+	// Overrides for shader warmup times to use when loading or to force a flush.
+	int32 OverridePrecompileTime;
+	int32 OverridePredrawBatchTime;
+	
 	static FShaderCache* Cache;
 	static int32 GameVersion;
 	static int32 bUseShaderCaching;
@@ -967,9 +990,17 @@ private:
 	static int32 bUseShaderDrawLog;
 	static int32 PredrawBatchTime;
 	static int32 bUseShaderBinaryCache;
+	static int32 bUseAsyncShaderPrecompilation;
+	static int32 TargetPrecompileFrameTime;
+	static int32 AccelPredrawBatchTime;
+	static int32 AccelTargetPrecompileFrameTime;
 	static FAutoConsoleVariableRef CVarUseShaderCaching;
 	static FAutoConsoleVariableRef CVarUseShaderPredraw;
 	static FAutoConsoleVariableRef CVarUseShaderDrawLog;
 	static FAutoConsoleVariableRef CVarPredrawBatchTime;
 	static FAutoConsoleVariableRef CVarUseShaderBinaryCache;
+	static FAutoConsoleVariableRef CVarUseAsyncShaderPrecompilation;
+	static FAutoConsoleVariableRef CVarTargetPrecompileFrameTime;
+	static FAutoConsoleVariableRef CVarAccelPredrawBatchTime;
+	static FAutoConsoleVariableRef CVarAccelTargetPrecompileFrameTime;
 };
