Fix another signed/unsigned warning that didn't appear under Clang but MSVC noticed.

diff --git a/Engine/Source/Runtime/ShaderCore/Private/ShaderCache.cpp b/Engine/Source/Runtime/ShaderCore/Private/ShaderCache.cpp
--- a/Engine/Source/Runtime/ShaderCore/Private/ShaderCache.cpp
+++ b/Engine/Source/Runtime/ShaderCore/Private/ShaderCache.cpp
@@ -1010,7 +1010,7 @@ void FShaderCache::InternalPreDrawShaders(FRHICommandList& RHICmdList, float Del
 		{
 			SET_DWORD_STAT(STATGROUP_NumToPrecompile, NumShadersToCompile);
 			
-			for( uint32 Index = 0; (GetTargetPrecompileFrameTime() == -1 || NumCompiled < NumShadersToCompile) && Index < ShadersToPrecompile.Num(); ++Index )
+			for( uint32 Index = 0; (GetTargetPrecompileFrameTime() == -1 || NumCompiled < NumShadersToCompile) && Index < (uint32)ShadersToPrecompile.Num(); ++Index )
 			{
 				FShaderCacheKey& Key = ShadersToPrecompile[Index];
 				TArray<uint8>& Code = CodeCache.Shaders[Key];
