In the shader-cache handle invalid resource references more gracefully, reporting the error and where it can be non-fatal assigning an invalid-state marker. We can't log draws while an invalid state is set so keep a count of them and ensure we only log draws with fully valid state. This should help debug cases where a resource lifetime isn't quite right without immediately crashing.

diff --git a/Engine/Source/Runtime/ShaderCore/Private/ShaderCache.cpp b/Engine/Source/Runtime/ShaderCore/Private/ShaderCache.cpp
--- a/Engine/Source/Runtime/ShaderCore/Private/ShaderCache.cpp
+++ b/Engine/Source/Runtime/ShaderCore/Private/ShaderCache.cpp
@@ -302,6 +302,7 @@ FShaderCache::FShaderCache()
 , CurrentShaderState(nullptr)
 , bIsPreDraw(false)
 , Options(SCO_Default)
+, InvalidResourceCount(0)
 , OverridePrecompileTime(0)
 , OverridePredrawBatchTime(0)
 {
@@ -837,16 +838,27 @@ void FShaderCache::InternalSetRenderTargets( uint32 NumSimultaneousRenderTargets
 		for( int32 RenderTargetIndex = NumSimultaneousRenderTargets - 1; RenderTargetIndex >= 0; --RenderTargetIndex )
 		{
 			FRHIRenderTargetView const& Target = NewRenderTargetsRHI[RenderTargetIndex];
+			InvalidResourceCount -= (uint32)(CurrentDrawKey.RenderTargets[RenderTargetIndex] == FShaderDrawKey::InvalidState);
 			if ( Target.Texture )
 			{
-				FShaderRenderTargetKey Key;
-				FSetElementId ID = FSetElementId::FromInteger(Textures.FindChecked(Target.Texture));
-				FShaderResourceKey TexKey = PlatformCache.Resources[ID];
-				Key.Texture = TexKey.Tex;
-				check(Key.Texture.MipLevels == Target.Texture->GetNumMips());
-				Key.MipLevel = Key.Texture.MipLevels > Target.MipIndex ? Target.MipIndex : 0;
-				Key.ArrayIndex = Target.ArraySliceIndex;
-				CurrentDrawKey.RenderTargets[RenderTargetIndex] = PlatformCache.RenderTargets.Add(Key).AsInteger();
+				int32* TexIndex = Textures.Find(Target.Texture);
+				if(TexIndex)
+				{
+					FSetElementId ID = FSetElementId::FromInteger(*TexIndex);
+					FShaderRenderTargetKey Key;
+					FShaderResourceKey TexKey = PlatformCache.Resources[ID];
+					Key.Texture = TexKey.Tex;
+					check(Key.Texture.MipLevels == Target.Texture->GetNumMips());
+					Key.MipLevel = Key.Texture.MipLevels > Target.MipIndex ? Target.MipIndex : 0;
+					Key.ArrayIndex = Target.ArraySliceIndex;
+					CurrentDrawKey.RenderTargets[RenderTargetIndex] = PlatformCache.RenderTargets.Add(Key).AsInteger();
+				}
+				else
+				{
+					UE_LOG(LogShaders, Warning, TEXT("Binding invalid texture %p to render target index %d, draw logging will be suspended until this is reset to a valid or null reference."), Target.Texture, RenderTargetIndex);
+					CurrentDrawKey.RenderTargets[RenderTargetIndex] = FShaderDrawKey::InvalidState;
+					InvalidResourceCount++;
+				}
 			}
 			else
 			{
@@ -854,13 +866,24 @@ void FShaderCache::InternalSetRenderTargets( uint32 NumSimultaneousRenderTargets
 			}
 		}
 		
+		InvalidResourceCount -= (uint32)(CurrentDrawKey.DepthStencilTarget == FShaderDrawKey::InvalidState);
 		if ( NewDepthStencilTargetRHI && NewDepthStencilTargetRHI->Texture )
 		{
-			FShaderRenderTargetKey Key;
-			FSetElementId ID = FSetElementId::FromInteger(Textures.FindChecked(NewDepthStencilTargetRHI->Texture));
-			FShaderResourceKey TexKey = PlatformCache.Resources[ID];
-			Key.Texture = TexKey.Tex;
-			CurrentDrawKey.DepthStencilTarget = PlatformCache.RenderTargets.Add(Key).AsInteger();
+			int32* TexIndex = Textures.Find(NewDepthStencilTargetRHI->Texture);
+			if(TexIndex)
+			{
+				FShaderRenderTargetKey Key;
+				FSetElementId ID = FSetElementId::FromInteger(*TexIndex);
+				FShaderResourceKey TexKey = PlatformCache.Resources[ID];
+				Key.Texture = TexKey.Tex;
+				CurrentDrawKey.DepthStencilTarget = PlatformCache.RenderTargets.Add(Key).AsInteger();
+			}
+			else
+			{
+				UE_LOG(LogShaders, Warning, TEXT("Binding invalid texture %p to denpth-stencil target, draw logging will be suspended until this is reset to a valid or null reference."), NewDepthStencilTargetRHI->Texture);
+				CurrentDrawKey.DepthStencilTarget = FShaderDrawKey::InvalidState;
+				InvalidResourceCount++;
+			}
 		}
 		else
 		{
@@ -874,10 +897,21 @@ void FShaderCache::InternalSetSamplerState(EShaderFrequency Frequency, uint32 In
 {
 	if ( bUseShaderDrawLog && !bIsPreDraw )
 	{
-		check(Index < GetFeatureLevelMaxTextureSamplers(GMaxRHIFeatureLevel));
+		checkf(Index < GetFeatureLevelMaxTextureSamplers(GMaxRHIFeatureLevel), TEXT("Attempting to bind sampler at index %d which exceeds RHI max. %d"), Index, GetFeatureLevelMaxTextureSamplers(GMaxRHIFeatureLevel));
+		InvalidResourceCount -= (uint32)(CurrentDrawKey.SamplerStates[Frequency][Index] == FShaderDrawKey::InvalidState);
 		if ( State )
 		{
-			CurrentDrawKey.SamplerStates[Frequency][Index] = SamplerStates.FindChecked(State);
+			int32* SamplerIndex = SamplerStates.Find(State);
+			if(SamplerIndex)
+			{
+				CurrentDrawKey.SamplerStates[Frequency][Index] = *SamplerIndex;
+			}
+			else
+			{
+				UE_LOG(LogShaders, Warning, TEXT("Binding invalid sampler %p to shader stage %d index %d, draw logging will be suspended until this is reset to a valid or null reference."), State, Frequency, Index);
+				CurrentDrawKey.SamplerStates[Frequency][Index] = FShaderDrawKey::InvalidState;
+				InvalidResourceCount++;
+			}
 		}
 		else
 		{
@@ -891,7 +925,8 @@ void FShaderCache::InternalSetTexture(EShaderFrequency Frequency, uint32 Index,
 {
 	if ( bUseShaderDrawLog && !bIsPreDraw )
 	{
-		check(Index < MaxResources);
+		checkf(Index < MaxResources, TEXT("Attempting to texture bind at index %d which exceeds RHI max. %d"), Index, MaxResources);
+		InvalidResourceCount -= (uint32)(CurrentDrawKey.Resources[Frequency][Index] == FShaderDrawKey::InvalidState);
 		if ( State )
 		{
 			FShaderResourceKey Key;
@@ -903,8 +938,18 @@ void FShaderCache::InternalSetTexture(EShaderFrequency Frequency, uint32 Index,
 			}
 			
 			FShaderPlatformCache& PlatformCache = Caches.FindOrAdd(GMaxRHIShaderPlatform);
-			FSetElementId ID = FSetElementId::FromInteger(Textures.FindChecked(Tex));
-			CurrentDrawKey.Resources[Frequency][Index] = ID.AsInteger();
+			int32* TexIndex = Textures.Find(Tex);
+			if(TexIndex)
+			{
+				FSetElementId ID = FSetElementId::FromInteger(*TexIndex);
+				CurrentDrawKey.Resources[Frequency][Index] = ID.AsInteger();
+			}
+			else
+			{
+				UE_LOG(LogShaders, Warning, TEXT("Binding invalid texture %p to shader stage %d index %d, draw logging will be suspended until this is reset to a valid or null reference."), State, Frequency, Index);
+				CurrentDrawKey.Resources[Frequency][Index] = FShaderDrawKey::InvalidState;
+				InvalidResourceCount++;
+			}
 		}
 		else
 		{
@@ -918,13 +963,22 @@ void FShaderCache::InternalSetSRV(EShaderFrequency Frequency, uint32 Index, FSha
 {
 	if ( bUseShaderDrawLog && !bIsPreDraw )
 	{
-		check(Index < MaxResources);
+		checkf(Index < MaxResources, TEXT("Attempting to bind SRV at index %d which exceeds RHI max. %d"), Index, MaxResources);
+		InvalidResourceCount -= (uint32)(CurrentDrawKey.Resources[Frequency][Index] == FShaderDrawKey::InvalidState);
 		if ( SRV )
 		{
-			FShaderResourceKey Key = SRVs.FindChecked(SRV);
-			
-			FShaderPlatformCache& PlatformCache = Caches.FindOrAdd(GMaxRHIShaderPlatform);
-			CurrentDrawKey.Resources[Frequency][Index] = PlatformCache.Resources.Add(Key).AsInteger();
+			FShaderResourceKey* Key = SRVs.Find(SRV);
+			if(Key)
+			{
+				FShaderPlatformCache& PlatformCache = Caches.FindOrAdd(GMaxRHIShaderPlatform);
+				CurrentDrawKey.Resources[Frequency][Index] = PlatformCache.Resources.Add(*Key).AsInteger();
+			}
+			else
+			{
+				UE_LOG(LogShaders, Warning, TEXT("Binding invalid SRV %p to shader stage %d index %d, draw logging will be suspended until this is reset to a valid or null reference."), SRV, Frequency, Index);
+				CurrentDrawKey.Resources[Frequency][Index] = FShaderDrawKey::InvalidState;
+				InvalidResourceCount++;
+			}
 		}
 		else
 		{
@@ -943,7 +997,15 @@ void FShaderCache::InternalSetBoundShaderState(FBoundShaderStateRHIParamRef Stat
 		CurrentShaderState = State;
 		if ( State )
 		{
-			BoundShaderState = ShaderStates.FindChecked(State);
+			FShaderCacheBoundState* NewState = ShaderStates.Find(State);
+			if(NewState)
+			{
+				BoundShaderState = *NewState;
+			}
+			else
+			{
+				UE_LOG(LogShaders, Fatal, TEXT("Binding invalid bound-shader-state %p"), State);
+			}
 		}
 		CurrentDrawKey.Hash = 0;
 	}
@@ -964,7 +1026,7 @@ void FShaderCache::InternalSetViewport(uint32 MinX, uint32 MinY, float MinZ, uin
 
 void FShaderCache::InternalLogDraw(uint8 IndexType)
 {
-	if ( bUseShaderDrawLog && !bIsPreDraw )
+	if ( bUseShaderDrawLog && !bIsPreDraw && InvalidResourceCount == 0 )
 	{
 		FShaderPlatformCache& PlatformCache = Caches.FindOrAdd(GMaxRHIShaderPlatform);
 		CurrentDrawKey.IndexType = IndexType;
@@ -1469,6 +1531,8 @@ void FShaderCache::SetShaderSamplerTextures( FRHICommandList& RHICmdList, FShade
 	
 	for ( uint32 i = 0; i < GetFeatureLevelMaxTextureSamplers(GMaxRHIFeatureLevel); i++ )
 	{
+		checkf(DrawKey.SamplerStates[Frequency][i] != FShaderDrawKey::InvalidState, TEXT("Resource state cannot be 'InvalidState' as that indicates a resource lifetime error in the application."));
+		
 		if ( DrawKey.SamplerStates[Frequency][i] != FShaderDrawKey::NullState )
 		{
 			FSamplerStateInitializerRHI SamplerInit = PlatformCache.SamplerStates[FSetElementId::FromInteger(DrawKey.SamplerStates[Frequency][i])];
@@ -1479,6 +1543,8 @@ void FShaderCache::SetShaderSamplerTextures( FRHICommandList& RHICmdList, FShade
 
 	for ( uint32 i = 0; i < MaxResources; i++ )
 	{
+		checkf(DrawKey.Resources[Frequency][i] != FShaderDrawKey::InvalidState, TEXT("Resource state cannot be 'InvalidState' as that indicates a resource lifetime error in the application."));
+		
 		FShaderTextureBinding Bind;
 		if ( DrawKey.Resources[Frequency][i] != FShaderDrawKey::NullState )
 		{
@@ -1570,6 +1636,8 @@ void FShaderCache::PreDrawShader(FRHICommandList& RHICmdList, FShaderCacheBoundS
 			FRHIRenderTargetView RenderTargets[MaxSimultaneousRenderTargets];
 			for ( uint32 i = 0; i < MaxSimultaneousRenderTargets; i++ )
 			{
+				checkf(DrawKey.RenderTargets[i] != FShaderDrawKey::InvalidState, TEXT("Resource state cannot be 'InvalidState' as that indicates a resource lifetime error in the application."));
+				
 				if( DrawKey.RenderTargets[i] != FShaderDrawKey::NullState )
 				{
 					FShaderTextureBinding Bind;
@@ -1591,6 +1659,8 @@ void FShaderCache::PreDrawShader(FRHICommandList& RHICmdList, FShaderCacheBoundS
 			FRHIDepthRenderTargetView DepthStencilTarget;
 			if ( bDepthStencilTarget )
 			{
+				checkf(DrawKey.DepthStencilTarget != FShaderDrawKey::InvalidState, TEXT("Resource state cannot be 'InvalidState' as that indicates a resource lifetime error in the application."));
+				
 				FShaderTextureBinding Bind;
 				FShaderRenderTargetKey RTKey = PlatformCache.RenderTargets[FSetElementId::FromInteger(DrawKey.DepthStencilTarget)];
 				Bind.Texture = CreateRenderTarget(RTKey);
diff --git a/Engine/Source/Runtime/ShaderCore/Public/ShaderCache.h b/Engine/Source/Runtime/ShaderCore/Public/ShaderCache.h
--- a/Engine/Source/Runtime/ShaderCore/Public/ShaderCache.h
+++ b/Engine/Source/Runtime/ShaderCore/Public/ShaderCache.h
@@ -253,7 +253,8 @@ class SHADERCORE_API FShaderCache : public FTickableObjectRenderThread
 		{
 			MaxNumSamplers = 16,
 			MaxNumResources = 128,
-			NullState = ~0u
+			NullState = ~0u,
+			InvalidState = NullState - 1u
 		};
 		
 		FShaderDrawKey()
@@ -979,6 +980,10 @@ private:
 	EShaderCacheOptions Options;
 	uint8 MaxResources;
 	
+	// When the invalid resource count is greater than 0 no draw keys will be stored to prevent corrupting the shader cache.
+	// Warnings are emitted to indicate that the shader cache has encountered a resource lifetime error.
+	uint32 InvalidResourceCount;
+	
 	// Overrides for shader warmup times to use when loading or to force a flush.
 	int32 OverridePrecompileTime;
 	int32 OverridePredrawBatchTime;
