Try to fix a different MSVC warning, which unfortunately clang doesn't allow us to enable independently.

diff --git a/Engine/Source/Runtime/ShaderCore/Public/ShaderCache.h b/Engine/Source/Runtime/ShaderCore/Public/ShaderCache.h
--- a/Engine/Source/Runtime/ShaderCore/Public/ShaderCache.h
+++ b/Engine/Source/Runtime/ShaderCore/Public/ShaderCache.h
@@ -253,8 +253,8 @@ class SHADERCORE_API FShaderCache : public FTickableObjectRenderThread
 		{
 			MaxNumSamplers = 16,
 			MaxNumResources = 128,
-			NullState = ~0u,
-			InvalidState = NullState - 1u
+			NullState = ~(0u),
+			InvalidState = ~(1u)
 		};
 		
 		FShaderDrawKey()
