And more MSVC errors not enabled on Clang...

diff --git a/Engine/Source/Runtime/ShaderCore/Private/ShaderCache.cpp b/Engine/Source/Runtime/ShaderCore/Private/ShaderCache.cpp
--- a/Engine/Source/Runtime/ShaderCore/Private/ShaderCache.cpp
+++ b/Engine/Source/Runtime/ShaderCore/Private/ShaderCache.cpp
@@ -908,7 +908,7 @@ void FShaderCache::InternalSetSamplerState(EShaderFrequency Frequency, uint32 In
 			}
 			else
 			{
-				UE_LOG(LogShaders, Warning, TEXT("Binding invalid sampler %p to shader stage %d index %d, draw logging will be suspended until this is reset to a valid or null reference."), State, Frequency, Index);
+				UE_LOG(LogShaders, Warning, TEXT("Binding invalid sampler %p to shader stage %d index %d, draw logging will be suspended until this is reset to a valid or null reference."), State, (uint32)Frequency, Index);
 				CurrentDrawKey.SamplerStates[Frequency][Index] = FShaderDrawKey::InvalidState;
 				InvalidResourceCount++;
 			}
@@ -946,7 +946,7 @@ void FShaderCache::InternalSetTexture(EShaderFrequency Frequency, uint32 Index,
 			}
 			else
 			{
-				UE_LOG(LogShaders, Warning, TEXT("Binding invalid texture %p to shader stage %d index %d, draw logging will be suspended until this is reset to a valid or null reference."), State, Frequency, Index);
+				UE_LOG(LogShaders, Warning, TEXT("Binding invalid texture %p to shader stage %d index %d, draw logging will be suspended until this is reset to a valid or null reference."), State, (uint32)Frequency, Index);
 				CurrentDrawKey.Resources[Frequency][Index] = FShaderDrawKey::InvalidState;
 				InvalidResourceCount++;
 			}
@@ -975,7 +975,7 @@ void FShaderCache::InternalSetSRV(EShaderFrequency Frequency, uint32 Index, FSha
 			}
 			else
 			{
-				UE_LOG(LogShaders, Warning, TEXT("Binding invalid SRV %p to shader stage %d index %d, draw logging will be suspended until this is reset to a valid or null reference."), SRV, Frequency, Index);
+				UE_LOG(LogShaders, Warning, TEXT("Binding invalid SRV %p to shader stage %d index %d, draw logging will be suspended until this is reset to a valid or null reference."), SRV, (uint32)Frequency, Index);
 				CurrentDrawKey.Resources[Frequency][Index] = FShaderDrawKey::InvalidState;
 				InvalidResourceCount++;
 			}
