Text format correctness in the new ShaderCache errors.

diff --git a/Engine/Source/Runtime/ShaderCore/Private/ShaderCache.cpp b/Engine/Source/Runtime/ShaderCore/Private/ShaderCache.cpp
--- a/Engine/Source/Runtime/ShaderCore/Private/ShaderCache.cpp
+++ b/Engine/Source/Runtime/ShaderCore/Private/ShaderCache.cpp
@@ -897,7 +897,7 @@ void FShaderCache::InternalSetSamplerState(EShaderFrequency Frequency, uint32 In
 {
 	if ( bUseShaderDrawLog && !bIsPreDraw )
 	{
-		checkf(Index < GetFeatureLevelMaxTextureSamplers(GMaxRHIFeatureLevel), TEXT("Attempting to bind sampler at index %d which exceeds RHI max. %d"), Index, GetFeatureLevelMaxTextureSamplers(GMaxRHIFeatureLevel));
+		checkf(Index < GetFeatureLevelMaxTextureSamplers(GMaxRHIFeatureLevel), TEXT("Attempting to bind sampler at index %u which exceeds RHI max. %d"), Index, GetFeatureLevelMaxTextureSamplers(GMaxRHIFeatureLevel));
 		InvalidResourceCount -= (uint32)(CurrentDrawKey.SamplerStates[Frequency][Index] == FShaderDrawKey::InvalidState);
 		if ( State )
 		{
@@ -908,7 +908,7 @@ void FShaderCache::InternalSetSamplerState(EShaderFrequency Frequency, uint32 In
 			}
 			else
 			{
-				UE_LOG(LogShaders, Warning, TEXT("Binding invalid sampler %p to shader stage %d index %d, draw logging will be suspended until this is reset to a valid or null reference."), State, (uint32)Frequency, Index);
+				UE_LOG(LogShaders, Warning, TEXT("Binding invalid sampler %p to shader stage %u index %u, draw logging will be suspended until this is reset to a valid or null reference."), State, (uint32)Frequency, Index);
 				CurrentDrawKey.SamplerStates[Frequency][Index] = FShaderDrawKey::InvalidState;
 				InvalidResourceCount++;
 			}
@@ -925,7 +925,7 @@ void FShaderCache::InternalSetTexture(EShaderFrequency Frequency, uint32 Index,
 {
 	if ( bUseShaderDrawLog && !bIsPreDraw )
 	{
-		checkf(Index < MaxResources, TEXT("Attempting to texture bind at index %d which exceeds RHI max. %d"), Index, MaxResources);
+		checkf(Index < MaxResources, TEXT("Attempting to texture bind at index %u which exceeds RHI max. %d"), Index, MaxResources);
 		InvalidResourceCount -= (uint32)(CurrentDrawKey.Resources[Frequency][Index] == FShaderDrawKey::InvalidState);
 		if ( State )
 		{
@@ -946,7 +946,7 @@ void FShaderCache::InternalSetTexture(EShaderFrequency Frequency, uint32 Index,
 			}
 			else
 			{
-				UE_LOG(LogShaders, Warning, TEXT("Binding invalid texture %p to shader stage %d index %d, draw logging will be suspended until this is reset to a valid or null reference."), State, (uint32)Frequency, Index);
+				UE_LOG(LogShaders, Warning, TEXT("Binding invalid texture %p to shader stage %u index %u, draw logging will be suspended until this is reset to a valid or null reference."), State, (uint32)Frequency, Index);
 				CurrentDrawKey.Resources[Frequency][Index] = FShaderDrawKey::InvalidState;
 				InvalidResourceCount++;
 			}
@@ -963,7 +963,7 @@ void FShaderCache::InternalSetSRV(EShaderFrequency Frequency, uint32 Index, FSha
 {
 	if ( bUseShaderDrawLog && !bIsPreDraw )
 	{
-		checkf(Index < MaxResources, TEXT("Attempting to bind SRV at index %d which exceeds RHI max. %d"), Index, MaxResources);
+		checkf(Index < MaxResources, TEXT("Attempting to bind SRV at index %u which exceeds RHI max. %d"), Index, MaxResources);
 		InvalidResourceCount -= (uint32)(CurrentDrawKey.Resources[Frequency][Index] == FShaderDrawKey::InvalidState);
 		if ( SRV )
 		{
@@ -975,7 +975,7 @@ void FShaderCache::InternalSetSRV(EShaderFrequency Frequency, uint32 Index, FSha
 			}
 			else
 			{
-				UE_LOG(LogShaders, Warning, TEXT("Binding invalid SRV %p to shader stage %d index %d, draw logging will be suspended until this is reset to a valid or null reference."), SRV, (uint32)Frequency, Index);
+				UE_LOG(LogShaders, Warning, TEXT("Binding invalid SRV %p to shader stage %u index %u, draw logging will be suspended until this is reset to a valid or null reference."), SRV, (uint32)Frequency, Index);
 				CurrentDrawKey.Resources[Frequency][Index] = FShaderDrawKey::InvalidState;
 				InvalidResourceCount++;
 			}
