Fix deprecation warning

diff --git a/Engine/Source/Runtime/ShaderCore/Private/Shader.cpp b/Engine/Source/Runtime/ShaderCore/Private/Shader.cpp
--- a/Engine/Source/Runtime/ShaderCore/Private/Shader.cpp
+++ b/Engine/Source/Runtime/ShaderCore/Private/Shader.cpp
@@ -978,7 +978,7 @@ FShaderPipelineType::FShaderPipelineType(const TCHAR* InName, const FShaderType*
 	}
 	Stages.Add(VertexShader);
 
-	GlobalListLink.Link(GetTypeList());
+	GlobalListLink.LinkHead(GetTypeList());
 	GetNameToTypeMap().Add(FName(InName), this);
 
 	// This will trigger if an IMPLEMENT_SHADER_TYPE was in a module not loaded before InitializeShaderTypes
