Mergine changes from //depot/UE4 to //UE4/Main

diff --git a/Engine/Source/Runtime/ShaderCore/Private/Shader.cpp b/Engine/Source/Runtime/ShaderCore/Private/Shader.cpp
--- a/Engine/Source/Runtime/ShaderCore/Private/Shader.cpp
+++ b/Engine/Source/Runtime/ShaderCore/Private/Shader.cpp
@@ -18,6 +18,43 @@
 DEFINE_LOG_CATEGORY(LogShaders);
 
 
+/**
+ * Find the shader pipeline type with the given name.
+ * @return NULL if no type matched.
+ */
+inline const FShaderPipelineType* FindShaderPipelineType(FName TypeName)
+{
+	for (TLinkedList<FShaderPipelineType*>::TIterator ShaderPipelineTypeIt(FShaderPipelineType::GetTypeList()); ShaderPipelineTypeIt; ShaderPipelineTypeIt.Next())
+	{
+		if (ShaderPipelineTypeIt->GetFName() == TypeName)
+		{
+			return *ShaderPipelineTypeIt;
+		}
+	}
+	return nullptr;
+}
+
+
+/**
+ * Serializes a reference to a shader pipeline type.
+ */
+FArchive& operator<<(FArchive& Ar, const FShaderPipelineType*& TypeRef)
+{
+	if (Ar.IsSaving())
+	{
+		FName TypeName = TypeRef ? FName(TypeRef->Name) : NAME_None;
+		Ar << TypeName;
+	}
+	else if (Ar.IsLoading())
+	{
+		FName TypeName = NAME_None;
+		Ar << TypeName;
+		TypeRef = FindShaderPipelineType(TypeName);
+	}
+	return Ar;
+}
+
+
 void FShaderParameterMap::VerifyBindingsAreComplete(const TCHAR* ShaderTypeName, EShaderFrequency Frequency, FVertexFactoryType* InVertexFactoryType) const
 {
 #if WITH_EDITORONLY_DATA
@@ -159,33 +196,43 @@ TMap<FName, FShaderType*>& FShaderType::GetNameToTypeMap()
 	return *GShaderNameToTypeMap;
 }
 
-void FShaderType::GetOutdatedTypes(TArray<FShaderType*>& OutdatedShaderTypes, TArray<const FVertexFactoryType*>& OutdatedFactoryTypes)
+inline bool FShaderType::GetOutdatedCurrentType(TArray<FShaderType*>& OutdatedShaderTypes, TArray<const FVertexFactoryType*>& OutdatedFactoryTypes) const
 {
-	for(TLinkedList<FShaderType*>::TIterator It(GetTypeList()); It; It.Next())
+	bool bOutdated = false;
+	for (TMap<FShaderId, FShader*>::TConstIterator ShaderIt(ShaderIdMap);ShaderIt;++ShaderIt)
 	{
-		FShaderType* Type = *It;
-		for(TMap<FShaderId,FShader*>::TConstIterator ShaderIt(Type->ShaderIdMap);ShaderIt;++ShaderIt)
+		FShader* Shader = ShaderIt.Value();
+		const FVertexFactoryParameterRef* VFParameterRef = Shader->GetVertexFactoryParameterRef();
+		const FSHAHash& SavedHash = Shader->GetHash();
+		const FSHAHash& CurrentHash = GetSourceHash();
+		const bool bOutdatedShader = SavedHash != CurrentHash;
+		const bool bOutdatedVertexFactory =
+			VFParameterRef && VFParameterRef->GetVertexFactoryType() && VFParameterRef->GetVertexFactoryType()->GetSourceHash() != VFParameterRef->GetHash();
+
+		if (bOutdatedShader)
 		{
-			FShader* Shader = ShaderIt.Value();
-			const FVertexFactoryParameterRef* VFParameterRef = Shader->GetVertexFactoryParameterRef();
-			const FSHAHash& SavedHash = Shader->GetHash();
-			const FSHAHash& CurrentHash = Type->GetSourceHash();
-			const bool bOutdatedShader = SavedHash != CurrentHash;
-			const bool bOutdatedVertexFactory =
-				VFParameterRef && VFParameterRef->GetVertexFactoryType() && VFParameterRef->GetVertexFactoryType()->GetSourceHash() != VFParameterRef->GetHash();
-
-			if (bOutdatedShader)
-			{
-				OutdatedShaderTypes.AddUnique(Shader->Type);
-			}
+			OutdatedShaderTypes.AddUnique(Shader->Type);
+			bOutdated = true;
+		}
 
-			if (bOutdatedVertexFactory)
-			{
-				OutdatedFactoryTypes.AddUnique(VFParameterRef->GetVertexFactoryType());
-			}
+		if (bOutdatedVertexFactory)
+		{
+			OutdatedFactoryTypes.AddUnique(VFParameterRef->GetVertexFactoryType());
+			bOutdated = true;
 		}
 	}
 
+	return bOutdated;
+}
+
+void FShaderType::GetOutdatedTypes(TArray<FShaderType*>& OutdatedShaderTypes, TArray<const FVertexFactoryType*>& OutdatedFactoryTypes)
+{
+	for(TLinkedList<FShaderType*>::TIterator It(GetTypeList()); It; It.Next())
+	{
+		FShaderType* Type = *It;
+		Type->GetOutdatedCurrentType(OutdatedShaderTypes, OutdatedFactoryTypes);
+	}
+
 	for (int32 TypeIndex = 0; TypeIndex < OutdatedShaderTypes.Num(); TypeIndex++)
 	{
 		UE_LOG(LogShaders, Warning, TEXT("		Recompiling %s"), OutdatedShaderTypes[TypeIndex]->GetName());
@@ -249,6 +296,7 @@ const FSHAHash& FShaderType::GetSourceHash() const
 
 void FShaderType::Initialize(const TMap<FString, TArray<const TCHAR*> >& ShaderFileToUniformBufferVariables)
 {
+	//#todo-rco: Need to call this only when Initializing from a Pipeline once it's removed from the global linked list
 	if (!FPlatformProperties::RequiresCookedData())
 	{
 		for(TLinkedList<FShaderType*>::TIterator It(FShaderType::GetTypeList()); It; It.Next())
@@ -581,8 +629,9 @@ FShaderResourceId FShaderResource::GetId() const
 	return ShaderId;
 }
 
-FShaderId::FShaderId(const FSHAHash& InMaterialShaderMapHash, FVertexFactoryType* InVertexFactoryType, FShaderType* InShaderType, FShaderTarget InTarget)
+FShaderId::FShaderId(const FSHAHash& InMaterialShaderMapHash, const FShaderPipelineType* InShaderPipeline, FVertexFactoryType* InVertexFactoryType, FShaderType* InShaderType, FShaderTarget InTarget)
 	: MaterialShaderMapHash(InMaterialShaderMapHash)
+	, ShaderPipeline(InShaderPipeline)
 	, ShaderType(InShaderType)
 	, SourceHash(InShaderType->GetSourceHash())
 	, SerializationHistory(InShaderType->GetSerializationHistory())
@@ -609,6 +658,7 @@ FSelfContainedShaderId::FSelfContainedShaderId(const FShaderId& InShaderId)
 {
 	MaterialShaderMapHash = InShaderId.MaterialShaderMapHash;
 	VertexFactoryTypeName = InShaderId.VertexFactoryType ? InShaderId.VertexFactoryType->GetName() : TEXT("");
+	ShaderPipelineName = InShaderId.ShaderPipeline ? InShaderId.ShaderPipeline->GetName() : TEXT("");
 	VFSourceHash = InShaderId.VFSourceHash;
 	VFSerializationHistory = InShaderId.VFSerializationHistory ? *InShaderId.VFSerializationHistory : FSerializationHistory();
 	ShaderTypeName = InShaderId.ShaderType->GetName();
@@ -620,7 +670,6 @@ FSelfContainedShaderId::FSelfContainedShaderId(const FShaderId& InShaderId)
 bool FSelfContainedShaderId::IsValid()
 {
 	FShaderType** TypePtr = FShaderType::GetNameToTypeMap().Find(FName(*ShaderTypeName));
-
 	if (TypePtr && SourceHash == (*TypePtr)->GetSourceHash() && SerializationHistory == (*TypePtr)->GetSerializationHistory())
 	{
 		FVertexFactoryType* VFTypePtr = FVertexFactoryType::GetVFByName(VertexFactoryTypeName);
@@ -639,6 +688,7 @@ FArchive& operator<<(FArchive& Ar,class FSelfContainedShaderId& Ref)
 {
 	Ar << Ref.MaterialShaderMapHash 
 		<< Ref.VertexFactoryTypeName
+		<< Ref.ShaderPipelineName
 		<< Ref.VFSourceHash
 		<< Ref.VFSerializationHistory
 		<< Ref.ShaderTypeName
@@ -654,8 +704,9 @@ FArchive& operator<<(FArchive& Ar,class FSelfContainedShaderId& Ref)
  * This still needs to initialize members to safe values since FShaderType::GenerateSerializationHistory uses this constructor.
  */
 FShader::FShader() : 
-	VFType(NULL),
-	Type(NULL), 
+	ShaderPipeline(nullptr),
+	VFType(nullptr),
+	Type(nullptr), 
 	NumRefs(0),
 	SetParametersId(0),
 	Canary(ShaderMagic_Uninitialized)
@@ -670,6 +721,7 @@ FShader::FShader() :
  */
 FShader::FShader(const CompiledShaderInitializerType& Initializer):
 	MaterialShaderMapHash(Initializer.MaterialShaderMapHash),
+	ShaderPipeline(Initializer.ShaderPipeline),
 	VFType(Initializer.VertexFactoryType),
 	Type(Initializer.Type),
 	Target(Initializer.Target),
@@ -736,6 +788,7 @@ bool FShader::SerializeBase(FArchive& Ar, bool bShadersInline)
 
 	Ar << OutputHash;
 	Ar << MaterialShaderMapHash;
+	Ar << ShaderPipeline;
 	Ar << VFType;
 	Ar << VFSourceHash;
 	Ar << Type;
@@ -893,6 +946,7 @@ FShaderId FShader::GetId() const
 {
 	FShaderId ShaderId(Type->GetSerializationHistory());
 	ShaderId.MaterialShaderMapHash = MaterialShaderMapHash;
+	ShaderId.ShaderPipeline = ShaderPipeline;
 	ShaderId.VertexFactoryType = VFType;
 	ShaderId.VFSourceHash = VFSourceHash;
 	ShaderId.VFSerializationHistory = VFType ? VFType->GetSerializationHistory((EShaderFrequency)GetTarget().Frequency) : NULL;
@@ -950,10 +1004,12 @@ FShaderPipelineType::FShaderPipelineType(
 	const FShaderType* InHullShader,
 	const FShaderType* InDomainShader,
 	const FShaderType* InGeometryShader,
-	const FShaderType* InPixelShader) :
+	const FShaderType* InPixelShader,
+	bool bInShouldAddShadersToIndividualMap) :
 	Name(InName),
 	TypeName(Name),
-	GlobalListLink(this)
+	GlobalListLink(this),
+	bAddShadersToIndividualMap(bInShouldAddShadersToIndividualMap)
 {
 	checkf(Name && *Name, TEXT("Shader Pipeline Type requires a valid Name!"));
 
@@ -1001,6 +1057,12 @@ FShaderPipelineType::FShaderPipelineType(
 
 FShaderPipelineType::~FShaderPipelineType()
 {
+	if (!bAddShadersToIndividualMap)
+	{
+		//#todo-rco: Delete & remove each FShader
+		check(0);
+	}
+
 	GetNameToTypeMap().Remove(FName(Name));
 	GlobalListLink.Unlink();
 }
@@ -1021,6 +1083,24 @@ TLinkedList<FShaderPipelineType*>*& FShaderPipelineType::GetTypeList()
 	return GShaderPipelineList;
 }
 
+TArray<const FShaderPipelineType*> FShaderPipelineType::GetShaderPipelineTypesByFilename(const TCHAR* Filename)
+{
+	TArray<const FShaderPipelineType*> PipelineTypes;
+	for (TLinkedList<FShaderPipelineType*>::TIterator It(FShaderPipelineType::GetTypeList()); It; It.Next())
+	{
+		auto* PipelineType = *It;
+		for (auto* ShaderType : PipelineType->Stages)
+		{
+			if (FPlatformString::Strcmp(Filename, ShaderType->GetShaderFilename()) == 0)
+			{
+				PipelineTypes.AddUnique(PipelineType);
+				break;
+			}
+		}
+	}
+	return PipelineTypes;
+}
+
 void FShaderPipelineType::Initialize()
 {
 	check(!bInitialized);
@@ -1075,6 +1155,122 @@ void FShaderPipelineType::Uninitialize()
 	bInitialized = false;
 }
 
+void FShaderPipelineType::GetOutdatedTypes(TArray<FShaderType*>& OutdatedShaderTypes, TArray<const FShaderPipelineType*>& OutdatedShaderPipelineTypes, TArray<const FVertexFactoryType*>& OutdatedFactoryTypes)
+{
+	for (TLinkedList<FShaderPipelineType*>::TIterator It(FShaderPipelineType::GetTypeList()); It; It.Next())
+	{
+		const auto* PipelineType = *It;
+		auto& Stages = PipelineType->GetStages();
+		bool bOutdated = false;
+		for (const FShaderType* ShaderType : Stages)
+		{
+			bOutdated = ShaderType->GetOutdatedCurrentType(OutdatedShaderTypes, OutdatedFactoryTypes) || bOutdated;
+		}
+
+		if (bOutdated)
+		{
+			OutdatedShaderPipelineTypes.AddUnique(PipelineType);
+		}
+	}
+
+	for (int32 TypeIndex = 0; TypeIndex < OutdatedShaderPipelineTypes.Num(); TypeIndex++)
+	{
+		UE_LOG(LogShaders, Warning, TEXT("		Recompiling Pipeline %s"), OutdatedShaderPipelineTypes[TypeIndex]->GetName());
+	}
+}
+
+
+FShaderPipeline::FShaderPipeline(
+	const FShaderPipelineType* InPipelineType,
+	FShader* InVertexShader,
+	FShader* InHullShader,
+	FShader* InDomainShader,
+	FShader* InGeometryShader,
+	FShader* InPixelShader) :
+	PipelineType(InPipelineType),
+	VertexShader(InVertexShader),
+	HullShader(InHullShader),
+	DomainShader(InDomainShader),
+	GeometryShader(InGeometryShader),
+	PixelShader(InPixelShader)
+{
+	check(InPipelineType);
+	Validate();
+}
+
+FShaderPipeline::FShaderPipeline(const FShaderPipelineType* InPipelineType, const TArray<FShader*>& InStages) :
+	PipelineType(InPipelineType),
+	VertexShader(nullptr),
+	HullShader(nullptr),
+	DomainShader(nullptr),
+	GeometryShader(nullptr),
+	PixelShader(nullptr)
+{
+	check(InPipelineType);
+	for (FShader* Shader : InStages)
+	{
+		if (Shader)
+		{
+			switch (Shader->GetType()->GetFrequency())
+			{
+			case SF_Vertex:
+				check(!VertexShader);
+				VertexShader = Shader;
+				break;
+			case SF_Pixel:
+				check(!PixelShader);
+				PixelShader = Shader;
+				break;
+			case SF_Hull:
+				check(!HullShader);
+				HullShader = Shader;
+				break;
+			case SF_Domain:
+				check(!DomainShader);
+				DomainShader = Shader;
+				break;
+			case SF_Geometry:
+				check(!GeometryShader);
+				GeometryShader = Shader;
+				break;
+			default:
+				checkf(0, TEXT("Invalid stage %u found!"), Shader->GetType()->GetFrequency());
+				break;
+			}
+		}
+	}
+
+	Validate();
+}
+
+void FShaderPipeline::Validate()
+{
+	for (const FShaderType* Stage : PipelineType->GetStages())
+	{
+		switch (Stage->GetFrequency())
+		{
+		case SF_Vertex:
+			check(VertexShader && VertexShader->GetType() == Stage);
+			break;
+		case SF_Pixel:
+			check(PixelShader && PixelShader->GetType() == Stage);
+			break;
+		case SF_Hull:
+			check(HullShader && HullShader->GetType() == Stage);
+			break;
+		case SF_Domain:
+			check(DomainShader && DomainShader->GetType() == Stage);
+			break;
+		case SF_Geometry:
+			check(GeometryShader && GeometryShader->GetType() == Stage);
+			break;
+		default:
+			// Can never happen :)
+			break;
+		}
+	}
+}
+
 void DumpShaderStats(EShaderPlatform Platform, EShaderFrequency Frequency)
 {
 #if ALLOW_DEBUG_FILES
@@ -1095,6 +1291,8 @@ void DumpShaderStats(EShaderPlatform Platform, EShaderFrequency Frequency)
 	ShaderTypeViewer.AddColumn(TEXT("AvgSizePerInstance"));
 	ShaderTypeViewer.CycleRow();
 
+	//#todo-rco: Pipelines
+
 	for( TLinkedList<FShaderType*>::TIterator It(FShaderType::GetTypeList()); It; It.Next() )
 	{
 		const FShaderType* Type = *It;
diff --git a/Engine/Source/Runtime/ShaderCore/Private/ShaderCore.cpp b/Engine/Source/Runtime/ShaderCore/Private/ShaderCore.cpp
--- a/Engine/Source/Runtime/ShaderCore/Private/ShaderCore.cpp
+++ b/Engine/Source/Runtime/ShaderCore/Private/ShaderCore.cpp
@@ -243,6 +243,9 @@ static void GetAllShaderSourceFiles( TArray<FString>& ShaderSourceFiles )
 			AddShaderSourceFileEntry(ShaderSourceFiles,ShaderFilename);
 		}
 	}
+
+	//#todo-rco: No need to loop through Shader Pipeline Types (yet)
+
 	// also always add the MaterialTemplate.usf shader file
 	AddShaderSourceFileEntry(ShaderSourceFiles,FString(TEXT("MaterialTemplate")));
 	AddShaderSourceFileEntry(ShaderSourceFiles,FString(TEXT("Common")));
@@ -595,6 +598,15 @@ void FlushShaderFileCache()
 		TMap<FString, TArray<const TCHAR*> > ShaderFileToUniformBufferVariables;
 		BuildShaderFileToUniformBufferMap(ShaderFileToUniformBufferVariables);
 
+		for (TLinkedList<FShaderPipelineType*>::TConstIterator It(FShaderPipelineType::GetTypeList()); It; It.Next())
+		{
+			const auto& Stages = It->GetStages();
+			for (const FShaderType* ShaderType : Stages)
+			{
+				((FShaderType*)ShaderType)->FlushShaderFileCache(ShaderFileToUniformBufferVariables);
+			}
+		}
+
 		for(TLinkedList<FShaderType*>::TIterator It(FShaderType::GetTypeList()); It; It.Next())
 		{
 			It->FlushShaderFileCache(ShaderFileToUniformBufferVariables);
diff --git a/Engine/Source/Runtime/ShaderCore/Public/Shader.h b/Engine/Source/Runtime/ShaderCore/Public/Shader.h
--- a/Engine/Source/Runtime/ShaderCore/Public/Shader.h
+++ b/Engine/Source/Runtime/ShaderCore/Public/Shader.h
@@ -346,6 +346,9 @@ public:
 	 */ 
 	FSHAHash MaterialShaderMapHash;
 
+	/** Shader Pipeline linked to this shader, needed since a single shader might be used on different Pipelines. */
+	const FShaderPipelineType* ShaderPipeline;
+
 	/** 
 	 * Vertex factory type that the shader was created for, 
 	 * This is needed in the Id since a single shader type will be compiled for multiple vertex factories within a material shader map.
@@ -380,7 +383,7 @@ public:
 	{}
 
 	/** Creates an Id for the given material, vertex factory, shader type and target. */
-	SHADERCORE_API FShaderId(const FSHAHash& InMaterialShaderMapHash, FVertexFactoryType* InVertexFactoryType, FShaderType* InShaderType, FShaderTarget InTarget);
+	SHADERCORE_API FShaderId(const FSHAHash& InMaterialShaderMapHash, const FShaderPipelineType* InShaderPipeline, FVertexFactoryType* InVertexFactoryType, FShaderType* InShaderType, FShaderTarget InTarget);
 
 	friend inline uint32 GetTypeHash( const FShaderId& Id )
 	{
@@ -390,6 +393,7 @@ public:
 	friend bool operator==(const FShaderId& X, const FShaderId& Y)
 	{
 		return X.MaterialShaderMapHash == Y.MaterialShaderMapHash
+			&& X.ShaderPipeline == Y.ShaderPipeline
 			&& X.VertexFactoryType == Y.VertexFactoryType
 			&& X.VFSourceHash == Y.VFSourceHash
 			&& ((X.VFSerializationHistory == NULL && Y.VFSerializationHistory == NULL)
@@ -425,6 +429,8 @@ public:
 	 */
 	FString VertexFactoryTypeName;
 
+	FString ShaderPipelineName;
+
 	/** Used to detect changes to the vertex factory source files. */
 	FSHAHash VFSourceHash;
 
@@ -469,6 +475,7 @@ public:
 		const FSHAHash& OutputHash;
 		FShaderResource* Resource;
 		FSHAHash MaterialShaderMapHash;
+		const FShaderPipelineType* ShaderPipeline;
 		FVertexFactoryType* VertexFactoryType;
 
 		CompiledShaderInitializerType(
@@ -476,6 +483,7 @@ public:
 			const FShaderCompilerOutput& CompilerOutput,
 			FShaderResource* InResource,
 			const FSHAHash& InMaterialShaderMapHash,
+			const FShaderPipelineType* InShaderPipeline,
 			FVertexFactoryType* InVertexFactoryType
 			):
 			Type(InType),
@@ -485,6 +493,7 @@ public:
 			OutputHash(CompilerOutput.OutputHash),
 			Resource(InResource),
 			MaterialShaderMapHash(InMaterialShaderMapHash),
+			ShaderPipeline(InShaderPipeline),
 			VertexFactoryType(InVertexFactoryType)
 		{}
 	};
@@ -568,6 +577,7 @@ public:
 	inline FSHAHash GetOutputHash() const { return OutputHash; }
 	FShaderId GetId() const;
 	inline FVertexFactoryType* GetVertexFactoryType() const { return VFType; }
+	inline const FShaderPipelineType* GetShaderPipelineType() const { return ShaderPipeline; }
 	inline FSHAHash GetMaterialShaderMapHash() const { return MaterialShaderMapHash; }
 	inline int32 GetNumRefs() const { return NumRefs; }
 
@@ -710,6 +720,9 @@ private:
 	/** Hash of the material shader map this shader belongs to, stored so that an FShaderId can be constructed from this shader. */
 	FSHAHash MaterialShaderMapHash;
 
+	/** Shader pipeline this shader belongs to, stored so that an FShaderId can be constructed from this shader. */
+	const FShaderPipelineType* ShaderPipeline;
+
 	/** Vertex factory type this shader was created for, stored so that an FShaderId can be constructed from this shader. */
 	FVertexFactoryType* VFType;
 
@@ -765,6 +778,9 @@ public:
 	/** Gets a list of FShaderTypes whose source file no longer matches what that type was compiled with */
 	static void GetOutdatedTypes(TArray<FShaderType*>& OutdatedShaderTypes, TArray<const FVertexFactoryType*>& OutdatedFactoryTypes);
 
+	/** Returns true if the source file no longer matches what that type was compiled with */
+	bool GetOutdatedCurrentType(TArray<FShaderType*>& OutdatedShaderTypes, TArray<const FVertexFactoryType*>& OutdatedFactoryTypes) const;
+
 	/** Initialize FShaderType static members, this must be called before any shader types are created. */
 	static void Initialize(const TMap<FString, TArray<const TCHAR*> >& ShaderFileToUniformBufferVariables);
 
@@ -1001,13 +1017,15 @@ private:
 class SHADERCORE_API FShaderPipelineType
 {
 public:
+	// Set bInShouldShadersToIndividualMap to false if you want to skip adding the FShaderTypes from the shadermap individual shader type list
 	FShaderPipelineType(
 		const TCHAR* InName,
 		const FShaderType* InVertexShader,
 		const FShaderType* InHullShader,
 		const FShaderType* InDomainShader,
 		const FShaderType* InGeometryShader,
-		const FShaderType* InPixelShader);
+		const FShaderType* InPixelShader,
+		bool bInShouldAddShadersToIndividualMap = true);
 	virtual ~FShaderPipelineType();
 
 	FORCEINLINE bool HasTessellation() const { return AllStages[SF_Domain] != nullptr; }
@@ -1035,6 +1053,8 @@ public:
 	static void Initialize();
 	static void Uninitialize();
 
+	static TArray<const FShaderPipelineType*> GetShaderPipelineTypesByFilename(const TCHAR* Filename);
+
 	/** Serializes a shader type reference by name. */
 	SHADERCORE_API friend FArchive& operator<<(FArchive& Ar, const FShaderPipelineType*& Ref);
 
@@ -1046,6 +1066,10 @@ public:
 	bool IsGlobalTypePipeline() const { return Stages[0]->GetGlobalShaderType() != nullptr; }
 	bool IsMaterialTypePipeline() const { return Stages[0]->GetMaterialShaderType() != nullptr; }
 	bool IsMeshMaterialTypePipeline() const { return Stages[0]->GetMeshMaterialShaderType() != nullptr; }
+	bool ShouldAddShadersToIndividualMap() const { return bAddShadersToIndividualMap; }
+
+	/** Gets a list of FShaderTypes & PipelineTypes whose source file no longer matches what that type was compiled with */
+	static void GetOutdatedTypes(TArray<FShaderType*>& OutdatedShaderTypes, TArray<const FShaderPipelineType*>& ShaderPipelineTypesToFlush, TArray<const FVertexFactoryType*>& OutdatedFactoryTypes);
 
 protected:
 	const TCHAR* const Name;
@@ -1059,6 +1083,7 @@ protected:
 	TLinkedList<FShaderPipelineType*> GlobalListLink;
 
 	uint32 HashIndex;
+	bool bAddShadersToIndividualMap;
 
 	static bool bInitialized;
 };
@@ -1070,27 +1095,31 @@ class FShaderTypeDependency
 public:
 
 	FShaderTypeDependency() :
-		ShaderType(NULL)
+		ShaderType(nullptr),
+		ShaderPipelineType(nullptr)
 	{}
 
 	/** Shader type */
 	FShaderType* ShaderType;
 
+	const FShaderPipelineType* ShaderPipelineType;
+
 	/** Used to detect changes to the shader source files. */
 	FSHAHash SourceHash;
 
 	friend FArchive& operator<<(FArchive& Ar,class FShaderTypeDependency& Ref)
 	{
-		Ar << Ref.ShaderType << Ref.SourceHash;
+		Ar << Ref.ShaderType << Ref.ShaderPipelineType << Ref.SourceHash;
 		return Ar;
 	}
 
 	bool operator==(const FShaderTypeDependency& Reference) const
 	{
-		return ShaderType == Reference.ShaderType && SourceHash == Reference.SourceHash;
+		return ShaderType == Reference.ShaderType && ShaderPipelineType == Reference.ShaderPipelineType && SourceHash == Reference.SourceHash;
 	}
 };
 
+
 /** Used to compare two shader types by name. */
 class FCompareShaderTypes
 {																				
@@ -1107,12 +1136,172 @@ public:
 	}
 };
 
+
+/** Used to compare two shader pipeline types by name. */
+class FCompareShaderPipelineNameTypes
+{
+public:
+	/*FORCEINLINE*/ bool operator()(const FShaderPipelineType& A, const FShaderPipelineType& B) const
+	{
+		//#todo-rco: Avoid this by adding an FNullShaderPipelineType
+		bool bNullA = &A == nullptr;
+		bool bNullB = &B == nullptr;
+		if (bNullA && bNullB)
+		{
+			return false;
+		}
+		else if (bNullA)
+		{
+			return true;
+		}
+		else if (bNullB)
+		{
+			return false;
+		}
+
+
+		int32 AL = FCString::Strlen(A.GetName());
+		int32 BL = FCString::Strlen(B.GetName());
+		if (AL == BL)
+		{
+			return FCString::Strncmp(A.GetName(), B.GetName(), AL) > 0;
+		}
+		return AL > BL;
+	}
+};
+
+/** Used to compare two shader pipeline types by name. */
+class FCompareShaderPipelineTypes
+{
+public:
+	FORCEINLINE bool operator()(const FShaderPipelineType& A, const FShaderPipelineType& B) const
+	{
+		int32 AL = FCString::Strlen(A.GetName());
+		int32 BL = FCString::Strlen(B.GetName());
+		if (AL == BL)
+		{
+			int32 Result = FCString::Strncmp(A.GetName(), B.GetName(), AL);
+			if (Result != 0)
+			{
+				return Result > 0;
+			}
+
+			const auto& PipeA = A.GetStages();
+			const auto& PipeB = B.GetStages();
+
+			if (PipeA.Num() != PipeB.Num())
+			{
+				return PipeA.Num() > PipeB.Num();
+			}
+
+			bool bResult = false;
+			FCompareShaderTypes Comparator;
+			for (int32 Index = 0; Index < PipeA.Num(); ++Index)
+			{
+				bResult = Comparator(*PipeA[Index], *PipeB[Index]);
+				if (bResult != 0)
+				{
+					break;
+				}
+			}
+
+			return bResult;
+		}
+		return AL > BL;
+	}
+};
+
+// A Shader Pipeline instance with compiled stages
+class SHADERCORE_API FShaderPipeline : public FRefCountedObject
+{
+public:
+	const FShaderPipelineType* PipelineType;
+	TRefCountPtr<FShader> VertexShader;
+	TRefCountPtr<FShader> HullShader;
+	TRefCountPtr<FShader> DomainShader;
+	TRefCountPtr<FShader> GeometryShader;
+	TRefCountPtr<FShader> PixelShader;
+
+	FShaderPipeline(
+		const FShaderPipelineType* InPipelineType,
+		FShader* InVertexShader,
+		FShader* InHullShader,
+		FShader* InDomainShader,
+		FShader* InGeometryShader,
+		FShader* InPixelShader);
+
+	FShaderPipeline(const FShaderPipelineType* InPipelineType, const TArray<FShader*>& InStages);
+
+	// Find a shader inside the pipeline
+	template<typename ShaderType>
+	ShaderType* GetShader()
+	{
+		if (PixelShader && PixelShader->GetType() == &ShaderType::StaticType)
+		{
+			return (ShaderType*)PixelShader;
+		}
+		else if (VertexShader && VertexShader->GetType() == &ShaderType::StaticType)
+		{
+			return (ShaderType*)VertexShader;
+		}
+		else if (GeometryShader && GeometryShader->GetType() == &ShaderType::StaticType)
+		{
+			return (ShaderType*)GeometryShader;
+		}
+		else if (HullShader)
+		{
+			if (HullShader->GetType() == &ShaderType::StaticType)
+			{
+				return (ShaderType*)HullShader;
+			}
+			else if (DomainShader && DomainShader->GetType() == &ShaderType::StaticType)
+			{
+				return (ShaderType*)DomainShader;
+			}
+		}
+
+		return nullptr;
+	}
+
+	FShader* GetShader(EShaderFrequency Frequency)
+	{
+		switch (Frequency)
+		{
+		case SF_Vertex: return VertexShader;
+		case SF_Domain: return DomainShader;
+		case SF_Hull: return HullShader;
+		case SF_Geometry: return GeometryShader;
+		case SF_Pixel: return PixelShader;
+		default: check(0);
+		}
+
+		return nullptr;
+	}
+
+	const FShader* GetShader(EShaderFrequency Frequency) const
+	{
+		switch (Frequency)
+		{
+		case SF_Vertex: return VertexShader;
+		case SF_Domain: return DomainShader;
+		case SF_Hull: return HullShader;
+		case SF_Geometry: return GeometryShader;
+		case SF_Pixel: return PixelShader;
+		default: check(0);
+		}
+
+		return nullptr;
+	}
+
+	void Validate();
+};
+
+
 /** A collection of shaders of different types, but the same meta type. */
 template<typename ShaderMetaType>
 class TShaderMap
 {
 public:
-
 	/** Default constructor. */
 	TShaderMap() {}
 
@@ -1121,7 +1310,7 @@ public:
 	ShaderType* GetShader() const
 	{
 		const TRefCountPtr<FShader>* ShaderRef = Shaders.Find(&ShaderType::StaticType);
-		checkf(ShaderRef != NULL && *ShaderRef != NULL, TEXT("Failed to find shader type %s"), ShaderType::StaticType.GetName());
+		checkf(ShaderRef != NULL && *ShaderRef != nullptr, TEXT("Failed to find shader type %s"), ShaderType::StaticType.GetName());
 		return (ShaderType*)((*ShaderRef)->GetShaderChecked());
 	}
 
@@ -1129,14 +1318,14 @@ public:
 	FShader* GetShader(FShaderType* ShaderType) const
 	{
 		const TRefCountPtr<FShader>* ShaderRef = Shaders.Find(ShaderType);
-		return ShaderRef ? (*ShaderRef)->GetShaderChecked() : NULL;
+		return ShaderRef ? (*ShaderRef)->GetShaderChecked() : nullptr;
 	}
 
 	/** Finds the shader with the given type. */
 	bool HasShader(FShaderType* Type) const
 	{
 		const TRefCountPtr<FShader>* ShaderRef = Shaders.Find(Type);
-		return ShaderRef != NULL && *ShaderRef != NULL;
+		return ShaderRef != nullptr && *ShaderRef != nullptr;
 	}
 
 	inline const TMap<FShaderType*,TRefCountPtr<FShader> >& GetShaders() const
@@ -1144,10 +1333,10 @@ public:
 		return Shaders;
 	}
 
-	void AddShader(FShaderType* Type,FShader* Shader)
+	void AddShader(FShaderType* Type, FShader* Shader)
 	{
 		check(Type);
-		Shaders.Add((FShaderType*)Type,Shader);
+		Shaders.Add(Type, Shader);
 	}
 
 	/**
@@ -1159,8 +1348,18 @@ public:
 		Shaders.Remove(Type);
 	}
 
+
+	void RemoveShaderPipelineType(const FShaderPipelineType* ShaderPipelineType)
+	{
+		auto* Found = ShaderPipelines.Find(ShaderPipelineType);
+		if (Found)
+		{
+			ShaderPipelines.Remove(ShaderPipelineType);
+		}
+	}
+
 	/** Builds a list of the shaders in a shader map. */
-	void GetShaderList(TMap<FShaderId,FShader*>& OutShaders) const
+	void GetShaderList(TMap<FShaderId, FShader*>& OutShaders) const
 	{
 		for(TMap<FShaderType*,TRefCountPtr<FShader> >::TConstIterator ShaderIt(Shaders);ShaderIt;++ShaderIt)
 		{
@@ -1183,6 +1382,15 @@ public:
 			}
 		}
 
+		for (auto Pair : ShaderPipelines)
+		{
+			const FShaderPipeline* Pipeline = Pair.Value;
+			for (const FShaderType* ShaderType : Pair.Key->GetStages())
+			{
+				MaxTextureSamplers = FMath::Max(MaxTextureSamplers, Pipeline->GetShader(ShaderType->GetFrequency())->GetNumTextureSamplers());
+			}
+		}
+
 		return MaxTextureSamplers;
 	}
 
@@ -1200,48 +1408,72 @@ public:
 
 			// Sort the shaders by type name before saving, to make sure the saved result is binary equivalent to what is generated on other machines, 
 			// Which is a requirement of the Derived Data Cache.
-			TMap<FShaderType*, TRefCountPtr<FShader> > SortedShaders = Shaders;
+			auto SortedShaders = Shaders;
 			SortedShaders.KeySort(FCompareShaderTypes());
 
-			for (TMap<FShaderType*,TRefCountPtr<FShader> >::TIterator It(SortedShaders); It; ++It)
-			{
-				FShader* CurrentShader = It.Value();
-				FShaderType* Type = It.Key();
-
-				check(Type);
-				checkSlow(FName(Type->GetName()) != NAME_None);
-
-				Ar << Type;
-				int32 SkipOffset = Ar.Tell();
-
+			auto SerializedShaderForSaving = [&](FShader* CurrentShader)
 				{
+					int32 SkipOffset = Ar.Tell();
+
+					{
 #if WITH_EDITOR
-					FArchive::FScopeSetDebugSerializationFlags S(Ar,DSF_IgnoreDiff);
+						FArchive::FScopeSetDebugSerializationFlags S(Ar, DSF_IgnoreDiff);
 #endif
-					// Serialize a placeholder value, we will overwrite this with an offset to the end of the shader
-					Ar << SkipOffset;
-				}
+						// Serialize a placeholder value, we will overwrite this with an offset to the end of the shader
+						Ar << SkipOffset;
+					}
 
-				if (bHandleShaderKeyChanges)
-				{
-					FSelfContainedShaderId SelfContainedKey = CurrentShader->GetId();
-					Ar << SelfContainedKey;
-				}
+					if (bHandleShaderKeyChanges)
+					{
+						FSelfContainedShaderId SelfContainedKey = CurrentShader->GetId();
+						Ar << SelfContainedKey;
+					}
+
+					CurrentShader->SerializeBase(Ar, bInlineShaderResource);
 
-				CurrentShader->SerializeBase(Ar, bInlineShaderResource);
+					// Get the offset to the end of the shader's serialized data
+					int32 EndOffset = Ar.Tell();
+					// Seek back to the placeholder and write the end offset
+					// This allows us to skip over the shader's serialized data at load time without knowing how to deserialize it
+					// Which can happen with shaders that were available at cook time, but not on the target platform (shaders in editor module for example)
+					Ar.Seek(SkipOffset);
+					Ar << EndOffset;
+					Ar.Seek(EndOffset);
+				};
+
+			for (TMap<FShaderType*, TRefCountPtr<FShader> >::TIterator ShaderIt(SortedShaders); ShaderIt; ++ShaderIt)
+			{
+				FShaderType* Type = ShaderIt.Key();
+				check(Type);
+				checkSlow(FName(Type->GetName()) != NAME_None);
 
-				// Get the offset to the end of the shader's serialized data
-				int32 EndOffset = Ar.Tell();
-				// Seek back to the placeholder and write the end offset
-				// This allows us to skip over the shader's serialized data at load time without knowing how to deserialize it
-				// Which can happen with shaders that were available at cook time, but not on the target platform (shaders in editor module for example)
-				Ar.Seek(SkipOffset);
-				Ar << EndOffset;
-				Ar.Seek(EndOffset);
+				Ar << Type;
+				FShader* CurrentShader = ShaderIt.Value();
+				SerializedShaderForSaving(CurrentShader);
 			}
 
-			int32 NumPipelines = 0;
+			int32 NumPipelines = ShaderPipelines.Num();
 			Ar << NumPipelines;
+			// Sort the shader pipelines by type name before saving, to make sure the saved result is binary equivalent to what is generated on other machines, Which is a requirement of the Derived Data Cache.
+			auto SortedPipelines = ShaderPipelines;
+			SortedPipelines.KeySort(FCompareShaderPipelineNameTypes());
+			for (TMap<const FShaderPipelineType*, TRefCountPtr<FShaderPipeline> >::TConstIterator PipelineIt(SortedPipelines); PipelineIt; ++PipelineIt)
+			{
+				FShaderPipeline* CurrentPipeline = PipelineIt.Value();
+				const FShaderPipelineType* PipelineType = CurrentPipeline->PipelineType;
+				Ar << PipelineType;
+
+				auto& PipelineStages = PipelineType->GetStages();
+				int32 NumStages = PipelineStages.Num();
+				Ar << NumStages;
+				for (int32 Index = 0; Index < NumStages; ++Index)
+				{
+					auto* Shader = CurrentPipeline->GetShader(PipelineStages[Index]->GetFrequency());
+					FShaderType* Type = Shader->GetType();
+					Ar << Type;
+					SerializedShaderForSaving(Shader);
+				}
+			}
 		}
 
 		if (Ar.IsLoading())
@@ -1249,60 +1481,101 @@ public:
 			int32 NumShaders = 0;
 			Ar << NumShaders;
 
-			for (int32 ShaderIndex = 0; ShaderIndex < NumShaders; ShaderIndex++)
-			{
-				FShaderType* Type = NULL;
-				Ar << Type;
-
-				int32 EndOffset = 0;
-				Ar << EndOffset;
-
-				FSelfContainedShaderId SelfContainedKey;
-
-				if (bHandleShaderKeyChanges)
+			auto SerializeShaderForLoad = [&](FShaderType* Type) -> FShader*
 				{
-					Ar << SelfContainedKey;
-				}
+					int32 EndOffset = 0;
+					Ar << EndOffset;
 
-				if (Type 
-					// If we are handling shader key changes, only create the shader if the serialized key matches the key the shader would have if created
-					// This allows serialization changes between the save and load to be safely handled
-					&& (!bHandleShaderKeyChanges || SelfContainedKey.IsValid()))
-				{
-					FShader* Shader = Type->ConstructForDeserialization();
-					check(Shader != NULL);
-					Shader->SerializeBase(Ar, bInlineShaderResource);
+					FSelfContainedShaderId SelfContainedKey;
 
-					TRefCountPtr<FShader> ExistingShader = Type->FindShaderById(Shader->GetId());
+					if (bHandleShaderKeyChanges)
+					{
+						Ar << SelfContainedKey;
+					}
 
-					if (ExistingShader.IsValid())
+					FShader* Shader = nullptr;
+					if (Type
+						// If we are handling shader key changes, only create the shader if the serialized key matches the key the shader would have if created
+						// This allows serialization changes between the save and load to be safely handled
+						&& (!bHandleShaderKeyChanges || SelfContainedKey.IsValid()))
 					{
-						delete Shader;
-						Shader = ExistingShader.GetReference();
+						Shader = Type->ConstructForDeserialization();
+						check(Shader != nullptr);
+						Shader->SerializeBase(Ar, bInlineShaderResource);
+
+						TRefCountPtr<FShader> ExistingShader = Type->FindShaderById(Shader->GetId());
+
+						if (ExistingShader.IsValid())
+						{
+							delete Shader;
+							Shader = ExistingShader.GetReference();
+						}
+						else
+						{
+							// Register the shader now that it is valid, so that it can be reused
+							Shader->Register();
+						}
 					}
 					else
 					{
-						// Register the shader now that it is valid, so that it can be reused
-						Shader->Register();
+						// Skip over this shader's serialized data if the type doesn't exist
+						// This can happen with shader types in modules that were loaded during cooking but not at run time (editor)
+						Ar.Seek(EndOffset);
 					}
+					return Shader;
+				};
 
-					AddShader(Shader->GetType(), Shader);
-				}
-				else
+			for (int32 ShaderIndex = 0; ShaderIndex < NumShaders; ShaderIndex++)
+			{
+				FShaderType* Type = nullptr;
+				Ar << Type;
+
+				FShader* Shader = SerializeShaderForLoad(Type);
+				if (Shader)
 				{
-					// Skip over this shader's serialized data if the type doesn't exist
-					// This can happen with shader types in modules that were loaded during cooking but not at run time (editor)
-					Ar.Seek(EndOffset);
+					AddShader(Shader->GetType(), Shader);
 				}
 			}
 
 			int32 NumPipelines = 0;
 			Ar << NumPipelines;
+			for (int32 PipelineIndex = 0; PipelineIndex < NumPipelines; ++PipelineIndex)
+			{
+				const FShaderPipelineType* ShaderPipelineType = nullptr;
+				Ar << ShaderPipelineType;
+				int32 NumStages = 0;
+				Ar << NumStages;
+				TArray<FShader*> ShaderStages;
+				for (int32 Index = 0; Index < NumStages; ++Index)
+				{
+					FShaderType* Type = nullptr;
+					Ar << Type;
+					FShader* Shader = SerializeShaderForLoad(Type);
+					if (Shader)
+					{
+						ShaderStages.Add(Shader);
+					}
+				}
+
+				if (ShaderStages.Num() == ShaderPipelineType->GetStages().Num())
+				{
+					FShaderPipeline* ShaderPipeline = new FShaderPipeline(ShaderPipelineType, ShaderStages);
+					AddShaderPipeline(ShaderPipelineType, ShaderPipeline);
+				}
+				else
+				{
+					//#todo-rco: Log that the # of stages changed for this pipeline
+					for (FShader* Shader : ShaderStages)
+					{
+						delete Shader;
+					}
+				}
+			}
 		}
 	}
 
 	/** @return true if the map is empty */
-	bool IsEmpty()
+	bool IsEmpty() const
 	{
 		return Shaders.Num() == 0;
 	}
@@ -1312,15 +1585,37 @@ public:
 	{
 		return Shaders.Num();
 	}
-	
+
+	/** @return The number of shader pipelines in the map. */
+	uint32 GetNumShaderPipelines() const
+	{
+		return ShaderPipelines.Num();
+	}
+
 	/** clears out all shaders held in the map */
 	void Empty()
 	{
 		Shaders.Empty();
+		ShaderPipelines.Empty();
+	}
+
+	// Returns nullptr if not found
+	FShaderPipeline* GetShaderPipeline(const FShaderPipelineType* PipelineType) const
+	{
+		auto* Found = ShaderPipelines.Find(PipelineType);
+		return Found ? *Found : nullptr;
+	}
+
+	void AddShaderPipeline(const FShaderPipelineType* Type, FShaderPipeline* ShaderPipeline)
+	{
+		check(Type);
+		check(!ShaderPipeline || ShaderPipeline->PipelineType == Type);
+		ShaderPipelines.Add(Type, ShaderPipeline);
 	}
 
 protected:
-	TMap<FShaderType*,TRefCountPtr<FShader> > Shaders;
+	TMap<FShaderType*, TRefCountPtr<FShader> > Shaders;
+	TMap<const FShaderPipelineType*, TRefCountPtr<FShaderPipeline> > ShaderPipelines;
 };
 
 /** A reference which is initialized with the requested shader type from a shader map. */
@@ -1570,4 +1865,4 @@ extern SHADERCORE_API void DispatchIndirectComputeShader(
 extern SHADERCORE_API const TArray<FName>& GetTargetShaderFormats();
 
 /** Appends to KeyString for all shaders. */
-extern SHADERCORE_API void ShaderMapAppendKeyString(EShaderPlatform Platform, FString& KeyString);
\ No newline at end of file
+extern SHADERCORE_API void ShaderMapAppendKeyString(EShaderPlatform Platform, FString& KeyString);
diff --git a/Engine/Source/Runtime/Slate/Private/Widgets/Input/SEditableLabel.cpp b/Engine/Source/Runtime/Slate/Private/Widgets/Input/SEditableLabel.cpp
new file mode 100644
--- /dev/null
+++ b/Engine/Source/Runtime/Slate/Private/Widgets/Input/SEditableLabel.cpp
@@ -0,0 +1,178 @@
+// Copyright 1998-2015 Epic Games, Inc. All Rights Reserved.
+
+#include "SlatePrivatePCH.h"
+#include "SEditableLabel.h"
+
+
+/* SEditableLabel interface
+ *****************************************************************************/
+
+void SEditableLabel::Construct(const FArguments& InArgs)
+{
+	CanEditAttribute = InArgs._CanEdit;
+	OnTextChanged = InArgs._OnTextChanged;
+	TextAttribute = InArgs._Text;
+
+	ChildSlot
+	[
+		SNew(SHorizontalBox)
+
+		+ SHorizontalBox::Slot()
+			.AutoWidth()
+			.VAlign(VAlign_Center)
+			[
+				SAssignNew(TextBlock, STextBlock)
+					.ColorAndOpacity(InArgs._ColorAndOpacity)
+					.Font(InArgs._Font)
+					.HighlightColor(InArgs._HighlightColor)
+					.HighlightShape(InArgs._HighlightShape)
+					.HighlightText(InArgs._HighlightText)
+					.MinDesiredWidth(InArgs._MinDesiredWidth)
+					.ShadowColorAndOpacity(InArgs._ShadowColorAndOpacity)
+					.ShadowOffset(InArgs._ShadowOffset)
+					.TextStyle(InArgs._TextStyle)
+					.Text(InArgs._Text)
+			]
+
+		+ SHorizontalBox::Slot()
+			.AutoWidth()
+			.VAlign(VAlign_Center) 
+			[
+				SAssignNew(EditableText, SEditableText)
+					.ClearKeyboardFocusOnCommit(true)
+					.ColorAndOpacity(InArgs._ColorAndOpacity)
+					.Font(InArgs._Font)
+					.MinDesiredWidth(InArgs._MinDesiredWidth)
+					.OnTextCommitted(this, &SEditableLabel::HandleEditableTextTextCommitted)
+					.RevertTextOnEscape(true)
+					.SelectAllTextOnCommit(false)
+					.SelectAllTextWhenFocused(true)
+					.Style(InArgs._EditableTextStyle)
+					.Text(InArgs._Text)
+					.VirtualKeyboardType(EKeyboardType::Keyboard_Number)
+					.Visibility(EVisibility::Collapsed)
+			]
+
+		+ SHorizontalBox::Slot()
+			.FillWidth(1.0f)
+			.Padding(4.0f, 0.0f, 0.0f, 0.0f)
+			.HAlign(HAlign_Left)
+			.VAlign(VAlign_Center)
+			[
+				SNew(SImage)
+					.Image(FCoreStyle::Get().GetBrush(TEXT("EditableLabel.EditIcon")))
+					.Visibility(this, &SEditableLabel::HandleIconVisibility)
+			]
+	];
+}
+
+
+void SEditableLabel::EnterTextMode()
+{
+	if (!CanEditAttribute.Get())
+	{
+		return;
+	}
+
+	TextBlock->SetVisibility(EVisibility::Collapsed);
+	EditableText->SetVisibility(EVisibility::Visible);
+}
+	
+
+void SEditableLabel::ExitTextMode()
+{
+	TextBlock->SetVisibility(EVisibility::Visible);
+	EditableText->SetVisibility(EVisibility::Collapsed);
+}
+
+
+/* SWidget interface
+ *****************************************************************************/
+
+bool SEditableLabel::HasKeyboardFocus() const
+{
+	// this label is considered focused when we are typing it text.
+	return SCompoundWidget::HasKeyboardFocus() ||
+		(EditableText.IsValid() && EditableText->HasKeyboardFocus());
+}
+
+
+FReply SEditableLabel::OnFocusReceived(const FGeometry& MyGeometry, const FFocusEvent& InFocusEvent)
+{
+	if ((InFocusEvent.GetCause() == EFocusCause::Navigation) ||
+		(InFocusEvent.GetCause() == EFocusCause::SetDirectly))
+	{
+		EnterTextMode();
+		return FReply::Handled().SetUserFocus(EditableText.ToSharedRef(), InFocusEvent.GetCause());
+	}
+
+	return FReply::Unhandled();
+}
+
+
+FReply SEditableLabel::OnKeyDown(const FGeometry& MyGeometry, const FKeyEvent& InKeyEvent)
+{
+	const FKey Key = InKeyEvent.GetKey();
+
+	if (Key == EKeys::Escape)
+	{
+		ExitTextMode();
+		return FReply::Handled().SetUserFocus(AsShared(), EFocusCause::Navigation);
+	}
+
+	if (Key == EKeys::Enter)
+	{
+		EnterTextMode();
+		return FReply::Handled().SetUserFocus(EditableText.ToSharedRef(), EFocusCause::Navigation);
+	}
+
+	return FReply::Unhandled();
+}
+
+
+FReply SEditableLabel::OnMouseButtonDown(const FGeometry& MyGeometry, const FPointerEvent& MouseEvent)
+{
+	if (MouseEvent.GetEffectingButton() == EKeys::LeftMouseButton)
+	{
+		// @todo Slate: implement some mechanism to not enter edit mode right away on mouse up
+	}
+
+	return FReply::Unhandled();
+}
+
+
+FReply SEditableLabel::OnMouseButtonUp(const FGeometry& MyGeometry, const FPointerEvent& MouseEvent )
+{
+	if (MouseEvent.GetEffectingButton() == EKeys::LeftMouseButton)
+	{
+		EnterTextMode();
+
+		return FReply::Handled();			
+	}
+
+	return FReply::Unhandled();
+}
+
+
+bool SEditableLabel::SupportsKeyboardFocus() const
+{
+	return true;
+}
+
+
+/* SEditableLabel callbacks
+ *****************************************************************************/
+
+void SEditableLabel::HandleEditableTextTextCommitted(const FText& NewText, ETextCommit::Type CommitInfo)
+{
+	ExitTextMode();
+	OnTextChanged.ExecuteIfBound(NewText);
+}
+
+
+EVisibility SEditableLabel::HandleIconVisibility() const
+{
+	return (IsHovered() && !EditableText->HasKeyboardFocus() && CanEditAttribute.Get())
+		? EVisibility::Visible
+		: EVisibility::Collapsed;
+}
diff --git a/Engine/Source/Runtime/SlateCore/Private/Styling/CoreStyle.cpp b/Engine/Source/Runtime/SlateCore/Private/Styling/CoreStyle.cpp
--- a/Engine/Source/Runtime/SlateCore/Private/Styling/CoreStyle.cpp
+++ b/Engine/Source/Runtime/SlateCore/Private/Styling/CoreStyle.cpp
@@ -414,6 +414,11 @@ TSharedRef<ISlateStyle> FCoreStyle::Create( const FName& InStyleSetName )
 		Style->Set( "RadioButton", BasicRadioButtonStyle );
 	}
 
+	// SEditableLabel defaults...
+	{
+		Style->Set("EditableLabel.EditIcon", new IMAGE_BRUSH("Icons/ellipsis_12x", Icon12x12));
+	}
+
 	// SEditableText defaults...
 	{
 		FSlateBrush* SelectionBackground = new BOX_BRUSH( "Common/EditableTextSelectionBackground",  FMargin(4.f/16.f) );
diff --git a/Engine/Source/Runtime/UMG/Private/Animation/MovieScene2DTransformTrackInstance.cpp b/Engine/Source/Runtime/UMG/Private/Animation/MovieScene2DTransformTrackInstance.cpp
--- a/Engine/Source/Runtime/UMG/Private/Animation/MovieScene2DTransformTrackInstance.cpp
+++ b/Engine/Source/Runtime/UMG/Private/Animation/MovieScene2DTransformTrackInstance.cpp
@@ -11,7 +11,7 @@ FMovieScene2DTransformTrackInstance::FMovieScene2DTransformTrackInstance( UMovie
 	PropertyBindings = MakeShareable(new FTrackInstancePropertyBindings(TransformTrack->GetPropertyName(), TransformTrack->GetPropertyPath()));
 }
 
-void FMovieScene2DTransformTrackInstance::Update( float Position, float LastPosition, const TArray<UObject*>& RuntimeObjects, class IMovieScenePlayer& Player, FMovieSceneSequenceInstance& SequenceInstance ) 
+void FMovieScene2DTransformTrackInstance::Update( float Position, float LastPosition, const TArray<UObject*>& RuntimeObjects, class IMovieScenePlayer& Player, FMovieSceneSequenceInstance& SequenceInstance, EMovieSceneUpdatePass UpdatePass ) 
 {
 	for(UObject* Object : RuntimeObjects)
 	{
diff --git a/Engine/Source/Runtime/UMG/Private/Animation/MovieScene2DTransformTrackInstance.h b/Engine/Source/Runtime/UMG/Private/Animation/MovieScene2DTransformTrackInstance.h
--- a/Engine/Source/Runtime/UMG/Private/Animation/MovieScene2DTransformTrackInstance.h
+++ b/Engine/Source/Runtime/UMG/Private/Animation/MovieScene2DTransformTrackInstance.h
@@ -15,7 +15,7 @@ public:
 	/** IMovieSceneTrackInstance interface */
 	virtual void SaveState(const TArray<UObject*>& RuntimeObjects, IMovieScenePlayer& Player, FMovieSceneSequenceInstance& SequenceInstance) override {}
 	virtual void RestoreState(const TArray<UObject*>& RuntimeObjects, IMovieScenePlayer& Player, FMovieSceneSequenceInstance& SequenceInstance) override {}
-	virtual void Update( float Position, float LastPosition, const TArray<UObject*>& RuntimeObjects, class IMovieScenePlayer& Player, FMovieSceneSequenceInstance& SequenceInstance ) override;
+	virtual void Update( float Position, float LastPosition, const TArray<UObject*>& RuntimeObjects, class IMovieScenePlayer& Player, FMovieSceneSequenceInstance& SequenceInstance, EMovieSceneUpdatePass UpdatePass ) override;
 	virtual void RefreshInstance( const TArray<UObject*>& RuntimeObjects, class IMovieScenePlayer& Player, FMovieSceneSequenceInstance& SequenceInstance ) override;
 	virtual void ClearInstance( IMovieScenePlayer& Player, FMovieSceneSequenceInstance& SequenceInstance ) override {};
 private:
diff --git a/Engine/Source/Runtime/UMG/Private/Animation/MovieSceneMarginTrackInstance.cpp b/Engine/Source/Runtime/UMG/Private/Animation/MovieSceneMarginTrackInstance.cpp
--- a/Engine/Source/Runtime/UMG/Private/Animation/MovieSceneMarginTrackInstance.cpp
+++ b/Engine/Source/Runtime/UMG/Private/Animation/MovieSceneMarginTrackInstance.cpp
@@ -12,7 +12,7 @@ FMovieSceneMarginTrackInstance::FMovieSceneMarginTrackInstance( UMovieSceneMargi
 	PropertyBindings = MakeShareable( new FTrackInstancePropertyBindings( MarginTrack->GetPropertyName(), MarginTrack->GetPropertyPath() ) );
 }
 
-void FMovieSceneMarginTrackInstance::Update( float Position, float LastPosition, const TArray<UObject*>& RuntimeObjects, class IMovieScenePlayer& Player, FMovieSceneSequenceInstance& SequenceInstance ) 
+void FMovieSceneMarginTrackInstance::Update( float Position, float LastPosition, const TArray<UObject*>& RuntimeObjects, class IMovieScenePlayer& Player, FMovieSceneSequenceInstance& SequenceInstance, EMovieSceneUpdatePass UpdatePass ) 
 {
 	for( UObject* Object : RuntimeObjects )
 	{
diff --git a/Engine/Source/Runtime/UMG/Public/Animation/UMGSequencePlayer.h b/Engine/Source/Runtime/UMG/Public/Animation/UMGSequencePlayer.h
--- a/Engine/Source/Runtime/UMG/Public/Animation/UMGSequencePlayer.h
+++ b/Engine/Source/Runtime/UMG/Public/Animation/UMGSequencePlayer.h
@@ -36,7 +36,7 @@ public:
 
 	/** IMovieScenePlayer interface */
 	virtual void GetRuntimeObjects( TSharedRef<FMovieSceneSequenceInstance> MovieSceneInstance, const FGuid& ObjectHandle, TArray< UObject* >& OutObjects ) const override;
-	virtual void UpdateCameraCut(UObject* ObjectToViewThrough, bool bNewCameraCut) const override {}
+	virtual void UpdateCameraCut(UObject* CameraObject, UObject* UnlockIfCameraObject) const override {}
 	virtual void SetViewportSettings(const TMap<FViewportClient*, EMovieSceneViewportParams>& ViewportParamsMap) override {}
 	virtual void GetViewportSettings(TMap<FViewportClient*, EMovieSceneViewportParams>& ViewportParamsMap) const override {}
 	virtual void AddOrUpdateMovieSceneInstance( class UMovieSceneSection& MovieSceneSection, TSharedRef<FMovieSceneSequenceInstance> InstanceToAdd ) override {}
