UE4 - Main - Fix serialization of really old materials

diff --git a/Engine/Source/Runtime/ShaderCore/Public/Shader.h b/Engine/Source/Runtime/ShaderCore/Public/Shader.h
--- a/Engine/Source/Runtime/ShaderCore/Public/Shader.h
+++ b/Engine/Source/Runtime/ShaderCore/Public/Shader.h
@@ -1109,7 +1109,12 @@ public:
 
 	friend FArchive& operator<<(FArchive& Ar,class FShaderTypeDependency& Ref)
 	{
-		Ar << Ref.ShaderType << Ref.ShaderPipelineType << Ref.SourceHash;
+		Ar << Ref.ShaderType;
+		if (Ar.UE4Ver() >= VER_UE4_PURGED_FMATERIAL_COMPILE_OUTPUTS)
+		{
+			Ar << Ref.ShaderPipelineType;
+		}
+		Ar << Ref.SourceHash;
 		return Ar;
 	}
 
