Orion-Main - Fix -ttrace for dumped debug shaders as levels 1-3 are allowed

diff --git a/Engine/Source/Runtime/ShaderCore/Private/Shader.cpp b/Engine/Source/Runtime/ShaderCore/Private/Shader.cpp
--- a/Engine/Source/Runtime/ShaderCore/Private/Shader.cpp
+++ b/Engine/Source/Runtime/ShaderCore/Private/Shader.cpp
@@ -17,6 +17,10 @@
 
 DEFINE_LOG_CATEGORY(LogShaders);
 
+static TAutoConsoleVariable<int32> CVarUsePipelines(
+	TEXT("r.ShaderPipelines"),
+	1,
+	TEXT("Enable using Shader pipelines."));
 
 /**
  * Find the shader pipeline type with the given name.
@@ -1488,7 +1492,7 @@ void ShaderMapAppendKeyString(EShaderPlatform Platform, FString& KeyString)
 		KeyString += (CVar && CVar->GetInt() != 0) ? TEXT("") : TEXT("_NoOpt");
 	}
 
-	if( Platform == SP_PS4 )
+	if (Platform == SP_PS4)
 	{
 		{
 			static const auto CVar = IConsoleManager::Get().FindTConsoleVariableDataInt(TEXT("r.PS4MixedModeShaderDebugInfo"));
@@ -1505,5 +1509,13 @@ void ShaderMapAppendKeyString(EShaderPlatform Platform, FString& KeyString)
 				KeyString += TEXT("_SDB");
 			}
 		}
+
+		{
+			static const auto CVar = IConsoleManager::Get().FindTConsoleVariableDataInt(TEXT("r.PS4UseTTrace"));
+			if (CVar && CVar->GetValueOnAnyThread() > 0)
+			{
+				KeyString += FString::Printf(TEXT("TT%d"), CVar->GetValueOnAnyThread());
+			}
+		}
 	}
 }
