Fix for xbox one compilation issue.

diff --git a/Engine/Source/Runtime/Engine/Public/GlobalShader.h b/Engine/Source/Runtime/Engine/Public/GlobalShader.h
--- a/Engine/Source/Runtime/Engine/Public/GlobalShader.h
+++ b/Engine/Source/Runtime/Engine/Public/GlobalShader.h
@@ -119,7 +119,7 @@ public:
 
 	ENGINE_API FGlobalShader(const ShaderMetaType::CompiledShaderInitializerType& Initializer);
 	
-	template<typename ShaderRHIParamRef, typename TRHICmdList = FRHICommandList>
+	template<typename ShaderRHIParamRef, typename TRHICmdList>
 	void SetParameters(TRHICmdList& RHICmdList, const ShaderRHIParamRef ShaderRHI, const FSceneView& View)
 	{
 		CheckShaderIsValid();
diff --git a/Engine/Source/Runtime/RHI/Public/RHICommandList.h b/Engine/Source/Runtime/RHI/Public/RHICommandList.h
--- a/Engine/Source/Runtime/RHI/Public/RHICommandList.h
+++ b/Engine/Source/Runtime/RHI/Public/RHICommandList.h
@@ -1327,6 +1327,12 @@ struct FRHIGraphicsWaitOnAsyncComputeJob : public FRHICommand<FRHIGraphicsWaitOn
 	RHI_API void Execute(FRHICommandListBase& CmdList);
 };
 
+template<> void FRHICommandSetShaderParameter<FComputeShaderRHIParamRef, ECmdList::ECompute>::Execute(FRHICommandListBase& CmdList);
+template<> void FRHICommandSetShaderUniformBuffer<FComputeShaderRHIParamRef, ECmdList::ECompute>::Execute(FRHICommandListBase& CmdList);
+template<> void FRHICommandSetShaderTexture<FComputeShaderRHIParamRef, ECmdList::ECompute>::Execute(FRHICommandListBase& CmdList);
+template<> void FRHICommandSetShaderResourceViewParameter<FComputeShaderRHIParamRef, ECmdList::ECompute>::Execute(FRHICommandListBase& CmdList);
+template<> void FRHICommandSetShaderSampler<FComputeShaderRHIParamRef, ECmdList::ECompute>::Execute(FRHICommandListBase& CmdList);
+
 class RHI_API FRHICommandList : public FRHICommandListBase
 {
 public:
diff --git a/Engine/Source/Runtime/Renderer/Public/SceneRenderTargetParameters.h b/Engine/Source/Runtime/Renderer/Public/SceneRenderTargetParameters.h
--- a/Engine/Source/Runtime/Renderer/Public/SceneRenderTargetParameters.h
+++ b/Engine/Source/Runtime/Renderer/Public/SceneRenderTargetParameters.h
@@ -63,7 +63,7 @@ class FDeferredPixelShaderParameters
 public:
 	void Bind(const FShaderParameterMap& ParameterMap);
 
-	template< typename ShaderRHIParamRef, typename TRHICmdList = FRHICommandList >
+	template< typename ShaderRHIParamRef, typename TRHICmdList >
 	void Set(
 		TRHICmdList& RHICmdList,
 		const ShaderRHIParamRef ShaderRHI,
diff --git a/Engine/Source/Runtime/ShaderCore/Public/ShaderParameterUtils.h b/Engine/Source/Runtime/ShaderCore/Public/ShaderParameterUtils.h
--- a/Engine/Source/Runtime/ShaderCore/Public/ShaderParameterUtils.h
+++ b/Engine/Source/Runtime/ShaderCore/Public/ShaderParameterUtils.h
@@ -168,7 +168,7 @@ inline void SetPixelShaderBool(
 /**
  * Sets the value of a shader texture parameter.  Template'd on shader type
  */
-template<typename ShaderTypeRHIParamRef, typename TRHICmdList = FRHICommandList>
+template<typename ShaderTypeRHIParamRef, typename TRHICmdList>
 FORCEINLINE void SetTextureParameter(
 	TRHICmdList& RHICmdList,
 	ShaderTypeRHIParamRef Shader,
@@ -209,7 +209,7 @@ FORCEINLINE void SetTextureParameter(
 /**
  * Sets the value of a shader texture parameter. Template'd on shader type.
  */
-template<typename ShaderTypeRHIParamRef, typename TRHICmdList = FRHICommandList>
+template<typename ShaderTypeRHIParamRef, typename TRHICmdList>
 FORCEINLINE void SetTextureParameter(
 	TRHICmdList& RHICmdList,
 	ShaderTypeRHIParamRef Shader,
@@ -249,7 +249,7 @@ FORCEINLINE void SetTextureParameter(
  * Sets the value of a shader surface parameter (e.g. to access MSAA samples).
  * Template'd on shader type (e.g. pixel shader or compute shader).
  */
-template<typename ShaderTypeRHIParamRef, typename TRHICmdList = FRHICommandList>
+template<typename ShaderTypeRHIParamRef, typename TRHICmdList>
 FORCEINLINE void SetTextureParameter(
 	TRHICmdList& RHICmdList,
 	ShaderTypeRHIParamRef Shader,
@@ -270,7 +270,7 @@ FORCEINLINE void SetTextureParameter(
 /**
  * Sets the value of a shader sampler parameter. Template'd on shader type.
  */
-template<typename ShaderTypeRHIParamRef, typename TRHICmdList = FRHICommandList>
+template<typename ShaderTypeRHIParamRef, typename TRHICmdList>
 FORCEINLINE void SetSamplerParameter(
 	TRHICmdList& RHICmdList,
 	ShaderTypeRHIParamRef Shader,
@@ -292,7 +292,7 @@ FORCEINLINE void SetSamplerParameter(
  * Sets the value of a shader resource view parameter
  * Template'd on shader type (e.g. pixel shader or compute shader).
  */
-template<typename ShaderTypeRHIParamRef, typename TRHICmdList = FRHICommandList>
+template<typename ShaderTypeRHIParamRef, typename TRHICmdList>
 FORCEINLINE void SetSRVParameter(
 	TRHICmdList& RHICmdList,
 	ShaderTypeRHIParamRef Shader,
@@ -313,7 +313,7 @@ FORCEINLINE void SetSRVParameter(
 /**
  * Sets the value of a unordered access view parameter
  */
-template<typename TRHICmdList = FRHICommandList>
+template<typename TRHICmdList>
 FORCEINLINE void SetUAVParameter(
 TRHICmdList& RHICmdList,
 	FComputeShaderRHIParamRef ComputeShader,
@@ -395,7 +395,7 @@ inline void SetLocalUniformBufferParameter(
 }
 
 /** Sets the value of a shader uniform buffer parameter to a uniform buffer containing the struct. */
-template<typename TShaderRHIRef, typename TRHICmdList = FRHICommandList>
+template<typename TShaderRHIRef, typename TRHICmdList>
 inline void SetUniformBufferParameter(
 	TRHICmdList& RHICmdList,
 	TShaderRHIRef Shader,
@@ -412,7 +412,7 @@ inline void SetUniformBufferParameter(
 }
 
 /** Sets the value of a shader uniform buffer parameter to a uniform buffer containing the struct. */
-template<typename TShaderRHIRef, typename TBufferStruct, typename TRHICmdList = FRHICommandList>
+template<typename TShaderRHIRef, typename TBufferStruct, typename TRHICmdList>
 inline void SetUniformBufferParameter(
 	TRHICmdList& RHICmdList,
 	TShaderRHIRef Shader,
@@ -429,7 +429,7 @@ inline void SetUniformBufferParameter(
 }
 
 /** Sets the value of a shader uniform buffer parameter to a uniform buffer containing the struct. */
-template<typename TShaderRHIRef, typename TBufferStruct, typename TRHICmdList = FRHICommandList>
+template<typename TShaderRHIRef, typename TBufferStruct, typename TRHICmdList>
 inline void SetUniformBufferParameter(
 	TRHICmdList& RHICmdList,
 	TShaderRHIRef Shader,
diff --git a/Engine/Source/Runtime/ShaderCore/Public/ShaderParameters.h b/Engine/Source/Runtime/ShaderCore/Public/ShaderParameters.h
--- a/Engine/Source/Runtime/ShaderCore/Public/ShaderParameters.h
+++ b/Engine/Source/Runtime/ShaderCore/Public/ShaderParameters.h
@@ -127,13 +127,13 @@ public:
 		return Ar << Parameter.SRVParameter << Parameter.UAVParameter;
 	}
 
-	template<typename TShaderRHIRef, typename TRHICmdList = FRHICommandList>
+	template<typename TShaderRHIRef, typename TRHICmdList>
 	inline void SetBuffer(TRHICmdList& RHICmdList, TShaderRHIRef Shader, const FRWBuffer& RWBuffer) const;
 
-	template<typename TShaderRHIRef, typename TRHICmdList = FRHICommandList>
+	template<typename TShaderRHIRef, typename TRHICmdList>
 	inline void SetTexture(TRHICmdList& RHICmdList, TShaderRHIRef Shader, const FTextureRHIParamRef Texture, FUnorderedAccessViewRHIParamRef UAV) const;
 
-	template<typename TRHICmdList = FRHICommandList>
+	template<typename TRHICmdList>
 	inline void UnsetUAV(TRHICmdList& RHICmdList, FComputeShaderRHIParamRef ComputeShader) const;
 
 private:
