Changed FShaderCompilerEnvironment::IncludeFileNameToContentsMap to use ANSICHAR instead of TCHAR strings to store values to make it use less memory.

diff --git a/Engine/Source/Developer/ShaderPreprocessor/Private/ShaderPreprocessor.cpp b/Engine/Source/Developer/ShaderPreprocessor/Private/ShaderPreprocessor.cpp
--- a/Engine/Source/Developer/ShaderPreprocessor/Private/ShaderPreprocessor.cpp
+++ b/Engine/Source/Developer/ShaderPreprocessor/Private/ShaderPreprocessor.cpp
@@ -76,7 +76,7 @@ private:
 			FString FileContents;
 			if (This->ShaderInput.Environment.IncludeFileNameToContentsMap.Contains(Filename))
 			{
-				FileContents = This->ShaderInput.Environment.IncludeFileNameToContentsMap.FindRef(Filename);
+				FileContents = FString(UTF8_TO_TCHAR(This->ShaderInput.Environment.IncludeFileNameToContentsMap.FindRef(Filename).GetData()));
 			}
 			else
 			{
diff --git a/Engine/Source/Runtime/Engine/Private/Materials/MaterialShared.cpp b/Engine/Source/Runtime/Engine/Private/Materials/MaterialShared.cpp
--- a/Engine/Source/Runtime/Engine/Private/Materials/MaterialShared.cpp
+++ b/Engine/Source/Runtime/Engine/Private/Materials/MaterialShared.cpp
@@ -1489,7 +1489,7 @@ bool FMaterial::BeginCompileShaderMap(
 		const FString MaterialShaderCode = MaterialTranslator.GetMaterialShaderCode();
 		const bool bSynchronousCompile = RequiresSynchronousCompilation() || !GShaderCompilingManager->AllowAsynchronousShaderCompiling();
 
-		MaterialEnvironment->IncludeFileNameToContentsMap.Add(TEXT("Material.usf"), MaterialShaderCode);
+		MaterialEnvironment->IncludeFileNameToContentsMap.Add(TEXT("Material.usf"), StringToArray<ANSICHAR>(*MaterialShaderCode, MaterialShaderCode.Len() + 1));
 
 		// Compile the shaders for the material.
 		NewShaderMap->Compile(this, ShaderMapId, MaterialEnvironment, NewCompilationOutput, Platform, bSynchronousCompile, bApplyCompletedShaderMapForRendering);
diff --git a/Engine/Source/Runtime/ShaderCore/Private/ShaderParameters.cpp b/Engine/Source/Runtime/ShaderCore/Private/ShaderParameters.cpp
--- a/Engine/Source/Runtime/ShaderCore/Private/ShaderParameters.cpp
+++ b/Engine/Source/Runtime/ShaderCore/Private/ShaderParameters.cpp
@@ -83,14 +83,16 @@ void FShaderUniformBufferParameter::ModifyCompilationEnvironment(const TCHAR* Pa
 {
 	const FString IncludeName = FString::Printf(TEXT("UniformBuffers/%s.usf"),ParameterName);
 	// Add the uniform buffer declaration to the compilation environment as an include: UniformBuffers/<ParameterName>.usf
-	OutEnvironment.IncludeFileNameToContentsMap.Add(
-		*IncludeName,
-		CreateUniformBufferShaderDeclaration(ParameterName,Struct,Platform)
-		);
-
-	FString& GeneratedUniformBuffersInclude = OutEnvironment.IncludeFileNameToContentsMap.FindOrAdd("GeneratedUniformBuffers.usf");
-	GeneratedUniformBuffersInclude += FString::Printf(TEXT("#include \"UniformBuffers/%s.usf\"") LINE_TERMINATOR, ParameterName);
+	FString Declaration = CreateUniformBufferShaderDeclaration(ParameterName,Struct,Platform);
+	OutEnvironment.IncludeFileNameToContentsMap.Add(IncludeName, StringToArray<ANSICHAR>(*Declaration, Declaration.Len() + 1));
 
+	TArray<ANSICHAR>& GeneratedUniformBuffersInclude = OutEnvironment.IncludeFileNameToContentsMap.FindOrAdd("GeneratedUniformBuffers.usf");
+	FString Include = FString::Printf(TEXT("#include \"UniformBuffers/%s.usf\"") LINE_TERMINATOR, ParameterName);
+	if (GeneratedUniformBuffersInclude.Num() > 0)
+	{
+		GeneratedUniformBuffersInclude.RemoveAt(GeneratedUniformBuffersInclude.Num() - 1);
+	}
+	GeneratedUniformBuffersInclude.Append(StringToArray<ANSICHAR>(*Include, Include.Len() + 1));
 	Struct.AddResourceTableEntries(OutEnvironment.ResourceTableMap, OutEnvironment.ResourceTableLayoutHashes);
 }
 
@@ -359,9 +361,10 @@ void FShaderType::AddReferencedUniformBufferIncludes(FShaderCompilerEnvironment&
 	{
 		check(It.Value().Declaration[Platform].Len() > 0);
 		UniformBufferIncludes += FString::Printf(TEXT("#include \"UniformBuffers/%s.usf\"") LINE_TERMINATOR, It.Key());
+		FString Declaration = It.Value().Declaration[Platform];
 		OutEnvironment.IncludeFileNameToContentsMap.Add(
 			*FString::Printf(TEXT("UniformBuffers/%s.usf"),It.Key()),
-			It.Value().Declaration[Platform]
+			StringToArray<ANSICHAR>(*Declaration, Declaration.Len() + 1)
 			);
 
 		for (TLinkedList<FUniformBufferStruct*>::TIterator StructIt(FUniformBufferStruct::GetStructList()); StructIt; StructIt.Next())
@@ -373,8 +376,12 @@ void FShaderType::AddReferencedUniformBufferIncludes(FShaderCompilerEnvironment&
 		}
 	}
 
-	FString& GeneratedUniformBuffersInclude = OutEnvironment.IncludeFileNameToContentsMap.FindOrAdd("GeneratedUniformBuffers.usf");
-	GeneratedUniformBuffersInclude += UniformBufferIncludes;
+	TArray<ANSICHAR>& GeneratedUniformBuffersInclude = OutEnvironment.IncludeFileNameToContentsMap.FindOrAdd("GeneratedUniformBuffers.usf");
+	if (GeneratedUniformBuffersInclude.Num() > 0)
+	{
+		GeneratedUniformBuffersInclude.RemoveAt(GeneratedUniformBuffersInclude.Num() - 1);
+	}
+	GeneratedUniformBuffersInclude.Append(StringToArray<ANSICHAR>(*UniformBufferIncludes, UniformBufferIncludes.Len() + 1));
 }
 
 void FVertexFactoryType::AddReferencedUniformBufferIncludes(FShaderCompilerEnvironment& OutEnvironment, FString& OutSourceFilePrefix, EShaderPlatform Platform)
@@ -392,9 +399,10 @@ void FVertexFactoryType::AddReferencedUniformBufferIncludes(FShaderCompilerEnvir
 	{
 		check(It.Value().Declaration[Platform].Len() > 0);
 		UniformBufferIncludes += FString::Printf(TEXT("#include \"UniformBuffers/%s.usf\"") LINE_TERMINATOR, It.Key());
+		FString Declaration = It.Value().Declaration[Platform];
 		OutEnvironment.IncludeFileNameToContentsMap.Add(
 			*FString::Printf(TEXT("UniformBuffers/%s.usf"),It.Key()),
-			It.Value().Declaration[Platform]
+			StringToArray<ANSICHAR>(*Declaration, Declaration.Len() + 1)
 		);
 
 		for (TLinkedList<FUniformBufferStruct*>::TIterator StructIt(FUniformBufferStruct::GetStructList()); StructIt; StructIt.Next())
@@ -406,6 +414,10 @@ void FVertexFactoryType::AddReferencedUniformBufferIncludes(FShaderCompilerEnvir
 		}
 	}
 
-	FString& GeneratedUniformBuffersInclude = OutEnvironment.IncludeFileNameToContentsMap.FindOrAdd("GeneratedUniformBuffers.usf");
-	GeneratedUniformBuffersInclude += UniformBufferIncludes;
+	TArray<ANSICHAR>& GeneratedUniformBuffersInclude = OutEnvironment.IncludeFileNameToContentsMap.FindOrAdd("GeneratedUniformBuffers.usf");
+	if (GeneratedUniformBuffersInclude.Num() > 0)
+	{
+		GeneratedUniformBuffersInclude.RemoveAt(GeneratedUniformBuffersInclude.Num() - 1);
+	}
+	GeneratedUniformBuffersInclude.Append(StringToArray<ANSICHAR>(*UniformBufferIncludes, UniformBufferIncludes.Len() + 1));
 }
diff --git a/Engine/Source/Runtime/ShaderCore/Public/ShaderCore.h b/Engine/Source/Runtime/ShaderCore/Public/ShaderCore.h
--- a/Engine/Source/Runtime/ShaderCore/Public/ShaderCore.h
+++ b/Engine/Source/Runtime/ShaderCore/Public/ShaderCore.h
@@ -318,7 +318,7 @@ inline FArchive& operator<<(FArchive& Ar, FResourceTableEntry& Entry)
 /** The environment used to compile a shader. */
 struct FShaderCompilerEnvironment : public FRefCountedObject
 {
-	TMap<FString,FString> IncludeFileNameToContentsMap;
+	TMap<FString,TArray<ANSICHAR>> IncludeFileNameToContentsMap;
 	TArray<uint32> CompilerFlags;
 	TMap<uint32,uint8> RenderTargetOutputFormatsMap;
 	TMap<FString,FResourceTableEntry> ResourceTableMap;
@@ -384,13 +384,17 @@ struct FShaderCompilerEnvironment : public FRefCountedObject
 	{
 		// Merge the include maps
 		// Merge the values of any existing keys
-		for (TMap<FString,FString>::TConstIterator It(Other.IncludeFileNameToContentsMap); It; ++It )
+		for (TMap<FString,TArray<ANSICHAR>>::TConstIterator It(Other.IncludeFileNameToContentsMap); It; ++It )
 		{
-			FString* ExistingContents = IncludeFileNameToContentsMap.Find(It.Key());
+			TArray<ANSICHAR>* ExistingContents = IncludeFileNameToContentsMap.Find(It.Key());
 
 			if (ExistingContents)
 			{
-				(*ExistingContents) += It.Value();
+				if (ExistingContents->Num() > 0)
+				{
+					ExistingContents->RemoveAt(ExistingContents->Num() - 1);
+				}
+				ExistingContents->Append(It.Value());
 			}
 			else
 			{
diff --git a/Engine/Source/Runtime/ShaderCore/Public/VertexFactory.h b/Engine/Source/Runtime/ShaderCore/Public/VertexFactory.h
--- a/Engine/Source/Runtime/ShaderCore/Public/VertexFactory.h
+++ b/Engine/Source/Runtime/ShaderCore/Public/VertexFactory.h
@@ -174,7 +174,7 @@ public:
 	{
 		// Set up the mapping from VertexFactory.usf to the vertex factory type's source code.
 		FString VertexFactoryIncludeString = FString::Printf( TEXT("#include \"%s.usf\""), GetShaderFilename() );
-		OutEnvironment.IncludeFileNameToContentsMap.Add(TEXT("VertexFactory.usf"), VertexFactoryIncludeString);
+		OutEnvironment.IncludeFileNameToContentsMap.Add(TEXT("VertexFactory.usf"), StringToArray<ANSICHAR>(*VertexFactoryIncludeString, VertexFactoryIncludeString.Len() + 1));
 
 		OutEnvironment.SetDefine(TEXT("HAS_PRIMITIVE_UNIFORM_BUFFER"), 1);
 
