Fixed a crash at startup in Mac cooked games for the case where bUseShaderBinaryCache is false (which is default), caused by skipped initialization of some of FShaderCache variables.

diff --git a/Engine/Source/Runtime/ShaderCore/Private/ShaderCache.cpp b/Engine/Source/Runtime/ShaderCore/Private/ShaderCache.cpp
--- a/Engine/Source/Runtime/ShaderCore/Private/ShaderCache.cpp
+++ b/Engine/Source/Runtime/ShaderCore/Private/ShaderCache.cpp
@@ -704,49 +704,54 @@ void FShaderCache::InitShaderCache(uint32 Options)
 	if(bUseShaderCaching)
 	{
 		Cache = new FShaderCache(Options);
-		}
 	}
+}
 
 void FShaderCache::LoadBinaryCache()
 {
+	if (!Cache)
+	{
+		return;
+	}
+
 	FShaderDrawKey::bTrackDrawResources = ShaderPlatformPrebindRequiresResource(GMaxRHIShaderPlatform);
 	
-	//We expect the RHI to be created at this point
-	if(Cache && Cache->DefaultCacheState == nullptr)
+	// We expect the RHI to be created at this point
+	if (Cache->DefaultCacheState == nullptr)
 	{
 		Cache->DefaultCacheState = Cache->InternalCreateOrFindCacheStateForContext(&GRHICommandList.GetImmediateCommandList().GetContext());
-}
+	}
+
+	FShaderCacheState* CacheState = Cache->InternalCreateOrFindCacheStateForContext(&GRHICommandList.GetImmediateCommandList().GetContext());
+
+	//I think we can get away without any locks in here specifically during the initial load phase as we are iterating the code cache
+
+	LoadTimeStart = FPlatformTime::Seconds();
+	Cache->ShadersToPrecompile = 0;
 
-	if(Cache && bUseShaderBinaryCache)
+	FString UserBinaryShaderFile = FPaths::GameSavedDir() / GShaderCacheFileName;
+	FString GameBinaryShaderFile = FPaths::GameContentDir() / GShaderCacheFileName;
+
+	// Try to load user cache, making sure that if we fail version test we still try game-content version.
+	bool bLoadedUserCache = LoadShaderCache(UserBinaryShaderFile, &Cache->Caches);
+
+	// Fallback to game-content version.
+	if ( !bLoadedUserCache )
+	{
+		LoadShaderCache(GameBinaryShaderFile, &Cache->Caches);
+	}
+
+	// Make sure relevant platform caches are already setup
+	for(size_t FeatureLevel = 0;FeatureLevel < ERHIFeatureLevel::Num;++FeatureLevel)
+	{
+		Cache->Caches.PlatformCaches.FindOrAdd(GShaderPlatformForFeatureLevel[FeatureLevel]);
+	}
+
+	//Prefetch the current platform cache
+	Cache->CurrentShaderPlatformCache = &Cache->Caches.PlatformCaches.FindOrAdd(GMaxRHIShaderPlatform);
+
+	if (bUseShaderBinaryCache)
 	{
-		FShaderCacheState* CacheState = Cache->InternalCreateOrFindCacheStateForContext(&GRHICommandList.GetImmediateCommandList().GetContext());
-		
-		//I think we can get away without any locks in here specifically during the initial load phase as we are iterating the code cache
-		
-		LoadTimeStart = FPlatformTime::Seconds();
-		Cache->ShadersToPrecompile = 0;
-		
-		FString UserBinaryShaderFile = FPaths::GameSavedDir() / GShaderCacheFileName;
-		FString GameBinaryShaderFile = FPaths::GameContentDir() / GShaderCacheFileName;
-		
-		// Try to load user cache, making sure that if we fail version test we still try game-content version.
-		bool bLoadedUserCache = LoadShaderCache(UserBinaryShaderFile, &Cache->Caches);
-			
-		// Fallback to game-content version.
-		if ( !bLoadedUserCache )
-			{
-			LoadShaderCache(GameBinaryShaderFile, &Cache->Caches);
-		}
-				
-		// Make sure relevant platform caches are already setup
-		for(size_t FeatureLevel = 0;FeatureLevel < ERHIFeatureLevel::Num;++FeatureLevel)
-				{
-			Cache->Caches.PlatformCaches.FindOrAdd(GShaderPlatformForFeatureLevel[FeatureLevel]);
-				}
-				
-		//Prefetch the current platform cache
-		Cache->CurrentShaderPlatformCache = &Cache->Caches.PlatformCaches.FindOrAdd(GMaxRHIShaderPlatform);
-		
 		if (FShaderCodeLibrary::GetShaderCount() > 0)
 		{
 			Cache->ShadersToPrecompile = FShaderCodeLibrary::GetShaderCount();
@@ -766,7 +771,7 @@ void FShaderCache::LoadBinaryCache()
 				FSHAHash Null;
 		
 				for (FShaderCodeLibraryPipeline const& Pipeline : *CachedPipelines)
-		{
+				{
 					PipelineKey.Hash = 0;
 					
 					PipelineKey.VertexShader.Platform = Platform;
@@ -822,63 +827,60 @@ void FShaderCache::LoadBinaryCache()
 					{
 						TSet<FShaderPipelineKey>& ShaderPipelines = Cache->Pipelines.FindOrAdd(PipelineKey.PixelShader);
 						ShaderPipelines.Add(PipelineKey);
+					}
+				}
 			}
 		}
-	}
-}
 
-		if (bUseShaderBinaryCache)
-{
-			for(uint32 i = 0; i < ERHIFeatureLevel::Num; i++)
-	{
-				if (GShaderPlatformForFeatureLevel[i] != SP_NumPlatforms)
+		for(uint32 i = 0; i < ERHIFeatureLevel::Num; i++)
 		{
-					const EShaderPlatform ShaderPlat = GShaderPlatformForFeatureLevel[i];
+			if (GShaderPlatformForFeatureLevel[i] != SP_NumPlatforms)
+			{
+				const EShaderPlatform ShaderPlat = GShaderPlatformForFeatureLevel[i];
 
-					// Regardless of the presence of the platform specific file we want an modifiable FShaderCacheLibrary to catch any outliers.
-					FShaderCacheLibrary* ShaderCacheLib = new FShaderCacheLibrary(ShaderPlat,GShaderCodeCacheFileName);
-					ShaderCacheLib->AddRef();
-					
-					// Try to load user cache, making sure that if we fail version test we still try game-content version.
-					bool bLoadedCache = ShaderCacheLib->Load(FPaths::GameSavedDir());
-					Cache->CodeCache.Add(ShaderPlat,ShaderCacheLib);
-					
-					Cache->CachedShaderLibraries.Add( ShaderPlat, ShaderCacheLib );
+				// Regardless of the presence of the platform specific file we want an modifiable FShaderCacheLibrary to catch any outliers.
+				FShaderCacheLibrary* ShaderCacheLib = new FShaderCacheLibrary(ShaderPlat,GShaderCodeCacheFileName);
+				ShaderCacheLib->AddRef();
+				
+				// Try to load user cache, making sure that if we fail version test we still try game-content version.
+				bool bLoadedCache = ShaderCacheLib->Load(FPaths::GameSavedDir());
+				Cache->CodeCache.Add(ShaderPlat,ShaderCacheLib);
+				
+				Cache->CachedShaderLibraries.Add( ShaderPlat, ShaderCacheLib );
 
-					if (bLoadedCache)
-{
-						for (auto const& ShaderPipelinePair : ShaderCacheLib->Pipelines)
-	{
-							TSet<FShaderPipelineKey>& ShaderPipelines = Cache->Pipelines.FindOrAdd(ShaderPipelinePair.Key);
-							ShaderPipelines.Append(ShaderPipelinePair.Value);
-						}
-						if (ShaderCacheLib->GetShaderCount())
-		{
-							Cache->ShadersToPrecompile += ShaderCacheLib->GetShaderCount();
-							Cache->ShaderLibraryPreCompileProgress.Add( ShaderCacheLib->CreateIterator() );
+				if (bLoadedCache)
+				{
+					for (auto const& ShaderPipelinePair : ShaderCacheLib->Pipelines)
+					{
+						TSet<FShaderPipelineKey>& ShaderPipelines = Cache->Pipelines.FindOrAdd(ShaderPipelinePair.Key);
+						ShaderPipelines.Append(ShaderPipelinePair.Value);
+					}
+					if (ShaderCacheLib->GetShaderCount())
+					{
+						Cache->ShadersToPrecompile += ShaderCacheLib->GetShaderCount();
+						Cache->ShaderLibraryPreCompileProgress.Add( ShaderCacheLib->CreateIterator() );
+					}
+				}
 			}
 		}
-		}
-	}
-}
 
 		bool const bUseAsync = bUseAsyncShaderPrecompilation != 0;
-{
+		{
 			double StartTime = FPlatformTime::Seconds();
 			bUseAsyncShaderPrecompilation = false;
 			
 			TArray<uint8> DummyCode;
 			
 			while(!bUseAsyncShaderPrecompilation && Cache->ShaderLibraryPreCompileProgress.Num() > 0)
-	{
+			{
 				TRefCountPtr<FRHIShaderLibrary::FShaderLibraryIterator> ShaderIterator =  Cache->ShaderLibraryPreCompileProgress[0];
 		
 				while (ShaderIterator->IsValid() && !bUseAsyncShaderPrecompilation)
-		{
+				{
 					auto LibraryEntry = **ShaderIterator;
 					
 					if(LibraryEntry.IsValid())
-		{
+					{
 						FShaderCacheKey Key;
 						Key.Platform = LibraryEntry.Platform;
 						Key.Frequency = LibraryEntry.Frequency;
@@ -890,27 +892,27 @@ void FShaderCache::LoadBinaryCache()
 			
 					double Duration = FPlatformTime::Seconds() - StartTime;
 					if(bUseAsync && InitialShaderLoadTime >= 0.0f && Duration >= InitialShaderLoadTime)
-			{
+					{
 						bUseAsyncShaderPrecompilation = bUseAsync;
-				}
+					}
 				
 					--(Cache->ShadersToPrecompile);
 					++(*ShaderIterator);
-			}
+				}
 		
 				//Invalid iterator - continue otherwise we're in async mode leave iterator in list
 				if (!ShaderIterator->IsValid())
-		{
+				{
 					Cache->ShaderLibraryPreCompileProgress.RemoveAt(0);
 				}
 				else
-			{
+				{
 					break;
+				}
 			}
 		}
 	}
 }
-}
 
 void FShaderCache::SaveBinaryCache(FString OutputDir, FName PlatformName)
 {
@@ -931,8 +933,8 @@ void FShaderCache::ShutdownShaderCache()
 	{
 		delete Cache;
 		Cache = nullptr;
-		}
 	}
+}
 
 FShaderCache::FShaderCache(uint32 InOptions)
 : FTickableObjectRenderThread(true)
