Merging fix for static analysis warning from 4.17 branch.

diff --git a/Engine/Source/Runtime/ShaderCore/Private/Shader.cpp b/Engine/Source/Runtime/ShaderCore/Private/Shader.cpp
--- a/Engine/Source/Runtime/ShaderCore/Private/Shader.cpp
+++ b/Engine/Source/Runtime/ShaderCore/Private/Shader.cpp
@@ -609,7 +609,11 @@ bool FShaderResource::ArePlatformsCompatible(EShaderPlatform CurrentPlatform, ES
 		// For Metal in Editor we can switch feature-levels, but not in cooked projects when using Metal shader librariss.
 		bool const bIsCurrentMetal = IsMetalPlatform(CurrentPlatform);
 		bool const bIsTargetMetal = IsMetalPlatform(TargetPlatform);
-		bool const bIsMetalCompatible = ((bIsCurrentMetal == bIsTargetMetal) && (!IsMetalPlatform(CurrentPlatform) || (WITH_EDITOR || (CurrentPlatform == TargetPlatform))));
+		bool const bIsMetalCompatible = (bIsCurrentMetal == bIsTargetMetal) 
+#if !WITH_EDITOR	// Static analysis doesn't like (|| WITH_EDITOR)
+			&& (!IsMetalPlatform(CurrentPlatform) || (CurrentPlatform == TargetPlatform))
+#endif
+			;
 		
 		bool const bIsCurrentOpenGL = IsOpenGLPlatform(CurrentPlatform);
 		bool const bIsTargetOpenGL = IsOpenGLPlatform(TargetPlatform);
@@ -634,14 +638,14 @@ void FShaderResource::InitRHI()
 
 	// we can't have this called on the wrong platform's shaders
 	if (!ArePlatformsCompatible(GMaxRHIShaderPlatform, (EShaderPlatform)Target.Platform))
- 	{
- 		if (FPlatformProperties::RequiresCookedData())
- 		{
- 			UE_LOG(LogShaders, Fatal, TEXT("FShaderResource::InitRHI got platform %s but it is not compatible with %s"), 
+	{
+		if (FPlatformProperties::RequiresCookedData())
+		{
+			UE_LOG(LogShaders, Fatal, TEXT("FShaderResource::InitRHI got platform %s but it is not compatible with %s"), 
 				*LegacyShaderPlatformToShaderFormat((EShaderPlatform)Target.Platform).ToString(), *LegacyShaderPlatformToShaderFormat(GMaxRHIShaderPlatform).ToString());
- 		}
- 		return;
- 	}
+		}
+		return;
+	}
 
 	TArray<uint8> UncompressedCode;
 	if (!bCodeInSharedLocation)
@@ -1845,13 +1849,13 @@ void ShaderMapAppendKeyString(EShaderPlatform Platform, FString& KeyString)
 			}
 		}
 	}
-    
-    // Encode the Metal standard into the shader compile options so that they recompile if the settings change.
-    if (IsMetalPlatform(Platform))
-    {
-        uint32 ShaderVersion = RHIGetShaderLanguageVersion(Platform);
-        KeyString += FString::Printf(TEXT("_MTLSTD%u_"), ShaderVersion);
-    }
+	
+	// Encode the Metal standard into the shader compile options so that they recompile if the settings change.
+	if (IsMetalPlatform(Platform))
+	{
+		uint32 ShaderVersion = RHIGetShaderLanguageVersion(Platform);
+		KeyString += FString::Printf(TEXT("_MTLSTD%u_"), ShaderVersion);
+	}
 
 	{
 		static const auto CVar = IConsoleManager::Get().FindTConsoleVariableDataInt(TEXT("r.StencilForLODDither"));
@@ -1899,24 +1903,24 @@ void ShaderMapAppendKeyString(EShaderPlatform Platform, FString& KeyString)
 		{
 			KeyString += TEXT("_8u");
 		}
-    }
-    
+	}
+	
 	{
 		static const auto CVar = IConsoleManager::Get().FindTConsoleVariableDataInt(TEXT("r.GPUSkin.Limit2BoneInfluences"));
 		if (CVar && CVar->GetValueOnAnyThread() != 0)
 		{
 			KeyString += TEXT("_2bi");
 		}
-    }
-    
-    if (IsMetalPlatform(Platform))
-    {
-        // Shaders built for archiving - for Metal that requires compiling the code in a different way so that we can strip it later
-        bool bArchive = false;
-        GConfig->GetBool(TEXT("/Script/UnrealEd.ProjectPackagingSettings"), TEXT("bSharedMaterialNativeLibraries"), bArchive, GGameIni);
-        if (bArchive)
-        {
-            KeyString += TEXT("_ARCHIVE");
-        }
-    }
+	}
+	
+	if (IsMetalPlatform(Platform))
+	{
+		// Shaders built for archiving - for Metal that requires compiling the code in a different way so that we can strip it later
+		bool bArchive = false;
+		GConfig->GetBool(TEXT("/Script/UnrealEd.ProjectPackagingSettings"), TEXT("bSharedMaterialNativeLibraries"), bArchive, GGameIni);
+		if (bArchive)
+		{
+			KeyString += TEXT("_ARCHIVE");
+		}
+	}
 }
