Remove obsolete CVar causing warnings

diff --git a/Engine/Config/BaseDeviceProfiles.ini b/Engine/Config/BaseDeviceProfiles.ini
--- a/Engine/Config/BaseDeviceProfiles.ini
+++ b/Engine/Config/BaseDeviceProfiles.ini
@@ -503,7 +503,6 @@ BaseProfileName=
 +CVars=r.Shaders.ZeroInitialise=1
 +CVars=r.Shaders.BoundsChecking=1
 +CVars=r.MorphTarget.Mode=0
-+CVars=r.Metal.TypedBuffers=1
 
 [MacClient DeviceProfile]
 DeviceType=Mac
diff --git a/Engine/Source/Runtime/ShaderCore/Private/Shader.cpp b/Engine/Source/Runtime/ShaderCore/Private/Shader.cpp
--- a/Engine/Source/Runtime/ShaderCore/Private/Shader.cpp
+++ b/Engine/Source/Runtime/ShaderCore/Private/Shader.cpp
@@ -1868,10 +1868,6 @@ void ShaderMapAppendKeyString(EShaderPlatform Platform, FString& KeyString)
 			static const auto CVar = IConsoleManager::Get().FindConsoleVariable(TEXT("r.Shaders.BoundsChecking"));
 			KeyString += (CVar && CVar->GetInt() != 0) ? TEXT("_BoundsChecking") : TEXT("");
 		}
-		{
-			static const auto CVar = IConsoleManager::Get().FindConsoleVariable(TEXT("r.Metal.TypedBuffers"));
-			KeyString += (CVar && CVar->GetInt() != 0) ? FString::Printf(TEXT("MTLTB%d"), CVar->GetInt()) : TEXT("");
-		}
 		
 		uint32 ShaderVersion = RHIGetShaderLanguageVersion(Platform);
 		KeyString += FString::Printf(TEXT("_MTLSTD%u_"), ShaderVersion);
