Duplicate 3642442:
Fix the native shader libraries again & undo the increase in cook time from changes to the way FShaderCodeLibrary was compiling & deduplicating the shaders.
- Remvoe the single linear array for accumulated shader code: reallocating this is tremendously expensive and will double the time taken to iteratively cook large projects.
- Uncompress the shader data for the native library system so that it actually works again.
- Fix some errant change to the Metal compiler that was trying to wedge the fully compiled library into the single Metal library which is 100% bogus.

diff --git a/Engine/Source/Developer/Apple/MetalShaderFormat/Private/MetalShaderCompiler.cpp b/Engine/Source/Developer/Apple/MetalShaderFormat/Private/MetalShaderCompiler.cpp
--- a/Engine/Source/Developer/Apple/MetalShaderFormat/Private/MetalShaderCompiler.cpp
+++ b/Engine/Source/Developer/Apple/MetalShaderFormat/Private/MetalShaderCompiler.cpp
@@ -1,5 +1,5 @@
 // Copyright 1998-2017 Epic Games, Inc. All Rights Reserved.
-// ..
+// ...
 
 #include "CoreMinimal.h"
 #include "MetalShaderFormat.h"
@@ -2048,8 +2048,9 @@ uint64 AppendShader_Metal(FName const& Format, FString const& WorkingDir, const
 				const uint8* SourceCodePtr = (uint8*)InShaderCode.GetData() + CodeOffset;
 				
 				// Copy the non-optional shader bytecode
-				int32 ObjectCodeDataSize = ShaderCode.GetActualShaderCodeSize() - CodeOffset;
-				TArrayView<const uint8> ObjectCodeArray(SourceCodePtr, ObjectCodeDataSize);
+				int32 ObjectCodeDataSize = 0;
+				uint8 const* Object = ShaderCode.FindOptionalDataAndSize('o', ObjectCodeDataSize);
+				TArrayView<const uint8> ObjectCodeArray(Object, ObjectCodeDataSize);
 				
 				// Object code segment
 				FString ObjFilename = WorkingDir / FString::Printf(TEXT("Main_%0.8x_%0.8x.o"), Header.SourceLen, Header.SourceCRC);
@@ -2195,7 +2196,7 @@ bool FinalizeLibrary_Metal(FName const& Format, FString const& WorkingDir, FStri
 					ExecRemoteProcess( *MetalArPath, *Params, &ReturnCode, &Results, &Errors );
 					bArchiveFileValid = RemoteFileExists(*ArchivePath);
 					
-                    if (ReturnCode != 0 || !bArchiveFileValid)
+					if (ReturnCode != 0 || !bArchiveFileValid)
 					{
 						UE_LOG(LogShaders, Error, TEXT("Archiving failed: metal-ar failed with code %d: %s"), ReturnCode, *Errors);
 						Params.Empty();
@@ -2239,8 +2240,8 @@ bool FinalizeLibrary_Metal(FName const& Format, FString const& WorkingDir, FStri
 				FString Params = FString::Printf(TEXT("-o=\"%s\" \"%s\""), *RemoteLibPath, *ArchivePath);
 					
 				ExecRemoteProcess( *MetalLibPath, *Params, &ReturnCode, &Results, &Errors );
-					
-                if(ReturnCode == 0)
+				
+				if(ReturnCode == 0)
                 {
                     // There is problem going to location with spaces using remote copy (at least on Mac no combination of \ and/or "" works) - work around this issue @todo investigate this further
                     FString LocalCopyLocation = FPaths::Combine(TEXT("/tmp"),FPaths::GetCleanFilename(LibraryPath));
diff --git a/Engine/Source/Runtime/ShaderCore/Private/ShaderCodeLibrary.cpp b/Engine/Source/Runtime/ShaderCore/Private/ShaderCodeLibrary.cpp
--- a/Engine/Source/Runtime/ShaderCore/Private/ShaderCodeLibrary.cpp
+++ b/Engine/Source/Runtime/ShaderCore/Private/ShaderCodeLibrary.cpp
@@ -510,12 +510,14 @@ struct FEditorShaderCodeArchive
 
 			FShaderCodeEntry Entry;
 			Entry.Size = InCode.Num();
-			Entry.Offset = ShaderCode.Num();
+			Entry.Offset = Offset;
 			Entry.UncompressedSize = UncompressedSize;
 			Entry.Frequency = Frequency;
+			Entry.LoadedCode = InCode;
+			
+			Offset += Entry.Size;
 			
 			Shaders.Add(Hash, Entry);
-			ShaderCode.Append(InCode);
 			bAdd = true;
 		}
 		return bAdd;
@@ -577,12 +579,11 @@ struct FEditorShaderCodeArchive
 					IFileManager::Get().MakeDirectory(*DebugPlatformDir, true);	
 										
 					TMap<FSHAHash, FShaderCodeEntry> StrippedShaders;
-					TArray<uint8> StrippedShaderCode;
-
+					uint32 TotalSize = 0;
 					for (const auto& Pair : Shaders)
 					{
 						TArray<uint8> CompressedCode;
-						CompressedCode.Append(ShaderCode.GetData() + Pair.Value.Offset, Pair.Value.Size);
+						CompressedCode.Append(Pair.Value.LoadedCode.GetData(), Pair.Value.Size);
 
 						int32 UncompressedSize = Pair.Value.UncompressedSize;
 					
@@ -598,23 +599,31 @@ struct FEditorShaderCodeArchive
 
 						FShaderCodeEntry StrippedEntry;
 						StrippedEntry.Size = CompressedCode.Num();
-						StrippedEntry.Offset = StrippedShaderCode.Num();
+						StrippedEntry.Offset = TotalSize;
 						StrippedEntry.UncompressedSize = UncompressedCode.Num();
 						StrippedEntry.Frequency = Pair.Value.Frequency;
+						StrippedEntry.LoadedCode = CompressedCode;
+						
+						TotalSize += StrippedEntry.Size;
 
 						StrippedShaders.Add(Pair.Key, StrippedEntry);
-						StrippedShaderCode.Append(CompressedCode);
 					}
-
+					
 					// Write stripped shader library
 					*FileWriter << StrippedShaders;
-					FileWriter->Serialize(StrippedShaderCode.GetData(), StrippedShaderCode.Num());
+					for (auto& Pair : StrippedShaders)
+					{
+						FileWriter->Serialize(Pair.Value.LoadedCode.GetData(), Pair.Value.Size);
+					}
 				}
 				else
 				{
 					// Write shader library
 					*FileWriter << Shaders;
-					FileWriter->Serialize(ShaderCode.GetData(), ShaderCode.Num());
+					for (auto& Pair : Shaders)
+					{
+						FileWriter->Serialize(Pair.Value.LoadedCode.GetData(), Pair.Value.Size);
+					}
 				}
 						
 				FileWriter->Close();
@@ -663,63 +672,19 @@ struct FEditorShaderCodeArchive
 			FString DebugPath = GetShaderCodeFilename(DebugShaderCodeDir, FormatName);
 			bOK = true;
 			
-			//Collect previous native cooked bytecode files into this shader files processing directory - keep the rest of the code simpler, don't overwrite in case dest file is newer
+			// Add the shaders to the archive.
+			for (auto& Pair : Shaders)
 			{
-				TArray<FString> NativeShaderFiles;
-				IFileManager::Get().FindFiles(NativeShaderFiles, *IntermediateCookedByteCodePath, TEXT("*.ushaderbytecode"));
+				FSHAHash& Hash = Pair.Key;
+				FShaderCodeEntry& Entry = Pair.Value;
 				
-				for (FString const& FileName : NativeShaderFiles)
+				TArray<uint8> UCode;
+				TArray<uint8>& UncompressedCode = FShaderLibraryHelperUncompressCode(Platform, Entry.UncompressedSize, Entry.LoadedCode, UCode);
+				
+				if(!Archive->AddShader(Entry.Frequency, Hash, UncompressedCode))
 				{
-					if (FileName.Len() > 2 && FileName[1] == TEXT('_'))
-					{
-						IFileManager::Get().Move(*(OutputPath / FileName), *(IntermediateCookedByteCodePath / FileName), false);
-					}
-				}
-			}
-			
-			TArray<FString> ShaderFiles;
-			IFileManager::Get().FindFiles(ShaderFiles, *OutputPath, TEXT("*.ushaderbytecode"));
-			
-			for (FString const& FileName : ShaderFiles)
-			{
-				if (FileName.Len() > 2 && FileName[1] == TEXT('_'))
-				{
-					FShaderCodeEntry Entry;
-					TCHAR FreqChar[2] = {0, 0};
-					FreqChar[0] = FileName[0];
-					Entry.Frequency = (EShaderFrequency)FCStringWide::Atoi(FreqChar);
-					check(Entry.Frequency < SF_NumFrequencies);
-					
-					FSHAHash Hash;
-					FString Name = FPaths::GetBaseFilename(FileName);
-					HexToBytes(Name.RightChop(2), Hash.Hash);
-					
-					FString Path = OutputPath / FileName;
-					FArchive* Ar = IFileManager::Get().CreateFileReader(*Path);
-					if (Ar)
-					{
-						uint32 UncompressedSize = 0;
-						TArray<uint8> CompressedCode;
-						
-						*Ar << UncompressedSize;
-						*Ar << CompressedCode;
-						
-						Ar->Close();
-						
-						TArray<uint8> UCode;
-						TArray<uint8>& UncompressedCode = FShaderLibraryHelperUncompressCode(Platform, UncompressedSize, CompressedCode, UCode);
-						
-						if(!Archive->AddShader(Entry.Frequency, Hash, UncompressedCode))
-						{
-							bOK = false;
-							break;
-						}
-					}
-					else
-					{
-						bOK = false;
-						break;
-					}
+					bOK = false;
+					break;
 				}
 			}
 			
@@ -730,19 +695,13 @@ struct FEditorShaderCodeArchive
 				//Always delete debug directory
 				IFileManager::Get().DeleteDirectory(*DebugShaderCodeDir, true, true);
 				
-				//Move files to intermediate dir for next iterative cook with overwwrite Move mode
-				if (bOK)
 				{
-					for (FString const& FileName : ShaderFiles)
-					{
-						if (FileName.Len() > 2 && FileName[1] == TEXT('_'))
-						{
-							IFileManager::Get().Move(*(IntermediateCookedByteCodePath / FileName), *(OutputPath / FileName), true);
-						}
-					}
-					
-					//We don't want to keep the shader code library shader files for native cooked content
-					IFileManager::Get().DeleteDirectory(*OutputPath, true, true);
+					FString OutputFilePath = GetCodeArchiveFilename(ShaderCodeDir, FormatName);
+					IFileManager::Get().Delete(*OutputFilePath);
+				}
+				{
+					FString OutputFilePath = GetPipelinesArchiveFilename(ShaderCodeDir, FormatName);
+					IFileManager::Get().Delete(*OutputFilePath);
 				}
 			}
 		}
@@ -752,7 +711,7 @@ struct FEditorShaderCodeArchive
 private:
 	FName FormatName;
 	TMap<FSHAHash, FShaderCodeEntry> Shaders;
-	TArray<uint8> ShaderCode;
+	uint32 Offset;
 	TSet<FShaderCodeLibraryPipeline> Pipelines;
 	const IShaderFormat* Format;
 };
