Fixes for compiler errors in nightly builds of VS2017 in /permissive- mode.

diff --git a/Engine/Plugins/Experimental/ProxyLODPlugin/Source/ProxyLOD/Private/ProxyLODMeshConvertUtils.h b/Engine/Plugins/Experimental/ProxyLODPlugin/Source/ProxyLOD/Private/ProxyLODMeshConvertUtils.h
--- a/Engine/Plugins/Experimental/ProxyLODPlugin/Source/ProxyLOD/Private/ProxyLODMeshConvertUtils.h
+++ b/Engine/Plugins/Experimental/ProxyLODPlugin/Source/ProxyLOD/Private/ProxyLODMeshConvertUtils.h
@@ -197,7 +197,7 @@ void ProxyLOD::MixedPolyMeshToAOSMesh(const FMixedPolyMesh& MixedPolyMesh, TAOSM
 			[&MixedPolyMesh, &DstAOSMesh, NumQuads](const ProxyLOD::FUIntRange& Range)
 		{
 			uint32* Indices = DstAOSMesh.Indexes;
-			for (uint32 t = Range.begin(), T = Range.end(); t < T; ++t)
+			for (uint32 t = Range.begin(), EndT = Range.end(); t < EndT; ++t)
 			{
 				const uint32 Offset = NumQuads * 6 + t * 3;
 				const openvdb::Vec3I& Tri = MixedPolyMesh.Triangles[t];
diff --git a/Engine/Source/Runtime/ShaderCore/Public/Shader.h b/Engine/Source/Runtime/ShaderCore/Public/Shader.h
--- a/Engine/Source/Runtime/ShaderCore/Public/Shader.h
+++ b/Engine/Source/Runtime/ShaderCore/Public/Shader.h
@@ -1690,7 +1690,7 @@ public:
 		check(bHasBeenRegistered);
 		uint32 MaxTextureSamplers = 0;
 
-		for (TMap<FShaderPrimaryKey,TRefCountPtr<FShader> >::TConstIterator ShaderIt(Shaders);ShaderIt;++ShaderIt)
+		for (typename TMap<FShaderPrimaryKey,TRefCountPtr<FShader> >::TConstIterator ShaderIt(Shaders);ShaderIt;++ShaderIt)
 		{
 			if (ShaderIt.Value())
 			{
@@ -1788,7 +1788,7 @@ public:
 			auto SortedShaders = Shaders;
 			SortedShaders.KeySort(TCompareShaderTypePermutation<FShaderType>());
 
-			for (TMap<FShaderPrimaryKey, TRefCountPtr<FShader> >::TIterator ShaderIt(SortedShaders); ShaderIt; ++ShaderIt)
+			for (typename TMap<FShaderPrimaryKey, TRefCountPtr<FShader> >::TIterator ShaderIt(SortedShaders); ShaderIt; ++ShaderIt)
 			{
 				FShader* CurrentShader = ShaderIt.Value();
 				FShaderType* Type = ShaderIt.Key().Type;
