RL - The uniformbuffer description is cached in multiple levels for shader compilation. Notable Caches are the ShaderDirectoryFileMapping the MMCP File cache and multiple other smaller file caches along the way. Instead of invalidating those caches I made the Uniform buffer layout completely immutable for all Platform varieties. Differences in lauout are triggered via defines as a progammer would have done it writing the code by hand. THe padding has been unified between GLSL and HLSL padding rules and theoretically work for both of them uniformly.

diff --git a/Engine/Shaders/Public/Platform.ush b/Engine/Shaders/Public/Platform.ush
--- a/Engine/Shaders/Public/Platform.ush
+++ b/Engine/Shaders/Public/Platform.ush
@@ -71,6 +71,10 @@
 #define COMPILER_SUPPORTS_ATTRIBUTES 0
 #endif
 
+#ifndef PLATFORM_SUPPORTS_SRV_UB
+#define PLATFORM_SUPPORTS_SRV_UB 0
+#endif
+
 #ifndef SM5_PROFILE
 #define SM5_PROFILE 0
 #endif
@@ -139,6 +143,12 @@
 #endif
 
 
+#if PLATFORM_SUPPORTS_SRV_UB
+#define PLATFORM_SUPPORTS_SRV_UB_MACRO(...) __VA_ARGS__
+#else
+#define PLATFORM_SUPPORTS_SRV_UB_MACRO(...)
+#endif
+
 // ---------------------------------------------------- Alternative floating point types
 
 #ifndef FORCE_FLOATS
diff --git a/Engine/Source/Runtime/ShaderCore/Private/Shader.cpp b/Engine/Source/Runtime/ShaderCore/Private/Shader.cpp
--- a/Engine/Source/Runtime/ShaderCore/Private/Shader.cpp
+++ b/Engine/Source/Runtime/ShaderCore/Private/Shader.cpp
@@ -153,10 +153,7 @@ FShaderType::FShaderType(
 	GetStreamOutElementsRef(InGetStreamOutElementsRef),
 	GlobalListLink(this)
 {
-	for (int32 Platform = 0; Platform < SP_NumPlatforms; Platform++)
-	{
-		bCachedUniformBufferStructDeclarations[Platform] = false;
-	}
+	bCachedUniformBufferStructDeclarations = false;
 
 	// This will trigger if an IMPLEMENT_SHADER_TYPE was in a module not loaded before InitializeShaderTypes
 	// Shader types need to be implemented in modules that are loaded before that
diff --git a/Engine/Source/Runtime/ShaderCore/Private/ShaderParameters.cpp b/Engine/Source/Runtime/ShaderCore/Private/ShaderParameters.cpp
--- a/Engine/Source/Runtime/ShaderCore/Private/ShaderParameters.cpp
+++ b/Engine/Source/Runtime/ShaderCore/Private/ShaderParameters.cpp
@@ -88,7 +88,7 @@ void FShaderUniformBufferParameter::ModifyCompilationEnvironment(const TCHAR* Pa
 	const FString IncludeName = FString::Printf(TEXT("/Engine/Generated/UniformBuffers/%s.ush"),ParameterName);
 	// Add the uniform buffer declaration to the compilation environment as an include: UniformBuffers/<ParameterName>.usf
 	FString Declaration;
-	CreateUniformBufferShaderDeclaration(ParameterName, Struct, Platform, Declaration);
+	CreateUniformBufferShaderDeclaration(ParameterName, Struct, Declaration);
 	OutEnvironment.IncludeVirtualPathToContentsMap.Add(IncludeName, Declaration);
 
 	FString& GeneratedUniformBuffersInclude = OutEnvironment.IncludeVirtualPathToContentsMap.FindOrAdd("/Engine/Generated/GeneratedUniformBuffers.ush");
@@ -148,8 +148,6 @@ struct FUniformBufferDecl
 static void CreateHLSLUniformBufferStructMembersDeclaration(
 	const FUniformBufferStruct& UniformBufferStruct, 
 	const FString& NamePrefix, 
-	ERHIFeatureLevel::Type MaxFeatureLevel,
-	bool bExplicitPadding, 
 	uint32 StructOffset, 
 	FUniformBufferDecl& Decl, 
 	uint32& HLSLBaseOffset)
@@ -159,6 +157,7 @@ static void CreateHLSLUniformBufferStructMembersDeclaration(
 	Decl.Initializer += TEXT("{");
 	int32 OpeningBraceLocPlusOne = Decl.Initializer.Len();
 
+	FString PreviousBaseTypeName = TEXT("float");
 	for (int32 MemberIndex = 0; MemberIndex < StructMembers.Num(); ++MemberIndex)
 	{
 		const FUniformBufferStruct::FMember& Member = StructMembers[MemberIndex];
@@ -174,7 +173,7 @@ static void CreateHLSLUniformBufferStructMembersDeclaration(
 			Decl.StructMembers += TEXT("struct {\r\n");
 			Decl.Initializer += TEXT(",");
 
-			CreateHLSLUniformBufferStructMembersDeclaration(*Member.GetStruct(), FString::Printf(TEXT("%s%s_"), *NamePrefix, Member.GetName()), MaxFeatureLevel, bExplicitPadding, StructOffset + Member.GetOffset(), Decl, HLSLBaseOffset);
+			CreateHLSLUniformBufferStructMembersDeclaration(*Member.GetStruct(), FString::Printf(TEXT("%s%s_"), *NamePrefix, Member.GetName()), StructOffset + Member.GetOffset(), Decl, HLSLBaseOffset);
 			Decl.StructMembers += FString::Printf(TEXT("} %s%s;\r\n"),Member.GetName(),*ArrayDim);
 		}
 		else if (IsUniformBufferResourceType(Member.GetBaseType()))
@@ -239,15 +238,12 @@ static void CreateHLSLUniformBufferStructMembersDeclaration(
 				check(HLSLBaseOffset < AbsoluteMemberOffset);
 				while(HLSLBaseOffset < AbsoluteMemberOffset)
 				{
-					if(bExplicitPadding)
-					{
-						Decl.ConstantBufferMembers += FString::Printf(TEXT("\tfloat1 _%sPrePadding%u;\r\n"), *NamePrefix, HLSLBaseOffset);
-					}
+					Decl.ConstantBufferMembers += FString::Printf(TEXT("\t%s PrePadding_%s%u;\r\n"), *PreviousBaseTypeName, *NamePrefix, HLSLBaseOffset);
 					HLSLBaseOffset += 4;
 				};
 				check(HLSLBaseOffset == AbsoluteMemberOffset);
 			}
-
+			PreviousBaseTypeName = BaseTypeName;
 			HLSLBaseOffset = AbsoluteMemberOffset + HLSLMemberSize;
 
 			// Generate the member declaration.
@@ -264,7 +260,15 @@ static void CreateHLSLUniformBufferStructMembersDeclaration(
 
 		if (IsUniformBufferResourceType(Member.GetBaseType()))
 		{
-			if (MaxFeatureLevel >= ERHIFeatureLevel::SM4 || Member.GetBaseType() != UBMT_SRV)
+			if (Member.GetBaseType() == UBMT_SRV)
+			{
+				// TODO: handle arrays?
+				FString ParameterName = FString::Printf(TEXT("%s%s"),*NamePrefix,Member.GetName());
+				Decl.ResourceMembers += FString::Printf(TEXT("PLATFORM_SUPPORTS_SRV_UB_MACRO( %s %s; ) \r\n"), Member.GetShaderType(), *ParameterName);
+				Decl.StructMembers += FString::Printf(TEXT("\tPLATFORM_SUPPORTS_SRV_UB_MACRO( %s %s; ) \r\n"), Member.GetShaderType(), Member.GetName());
+				Decl.Initializer += FString::Printf(TEXT(" PLATFORM_SUPPORTS_SRV_UB_MACRO( ,%s ) "), *ParameterName);
+			}
+			else
 			{
 				// TODO: handle arrays?
 				FString ParameterName = FString::Printf(TEXT("%s%s"),*NamePrefix,Member.GetName());
@@ -283,7 +287,7 @@ static void CreateHLSLUniformBufferStructMembersDeclaration(
 }
 
 /** Creates a HLSL declaration of a uniform buffer with the given structure. */
-static FString CreateHLSLUniformBufferDeclaration(const TCHAR* Name,const FUniformBufferStruct& UniformBufferStruct, ERHIFeatureLevel::Type MaxFeatureLevel, bool bExplicitPadding)
+static FString CreateHLSLUniformBufferDeclaration(const TCHAR* Name,const FUniformBufferStruct& UniformBufferStruct)
 {
 	// If the uniform buffer has no members, we don't want to write out anything.  Shader compilers throw errors when faced with empty cbuffers and structs.
 	if (UniformBufferStruct.GetMembers().Num() > 0)
@@ -291,7 +295,7 @@ static FString CreateHLSLUniformBufferDeclaration(const TCHAR* Name,const FUnifo
 		FString NamePrefix(FString(Name) + FString(TEXT("_")));
 		FUniformBufferDecl Decl;
 		uint32 HLSLBaseOffset = 0;
-		CreateHLSLUniformBufferStructMembersDeclaration(UniformBufferStruct, NamePrefix, MaxFeatureLevel, bExplicitPadding, 0, Decl, HLSLBaseOffset);
+		CreateHLSLUniformBufferStructMembersDeclaration(UniformBufferStruct, NamePrefix, 0, Decl, HLSLBaseOffset);
 
 		return FString::Printf(
 			TEXT("#ifndef __UniformBuffer_%s_Definition__\r\n")
@@ -321,37 +325,26 @@ static FString CreateHLSLUniformBufferDeclaration(const TCHAR* Name,const FUnifo
 	return FString(TEXT("\n"));
 }
 
-void CreateUniformBufferShaderDeclaration(const TCHAR* Name,const FUniformBufferStruct& UniformBufferStruct,EShaderPlatform Platform, FString& OutDeclaration)
+SHADERCORE_API void CreateUniformBufferShaderDeclaration(const TCHAR* Name,const FUniformBufferStruct& UniformBufferStruct, FString& OutDeclaration)
 {
-	ERHIFeatureLevel::Type MaxFeatureLevel = GetMaxSupportedFeatureLevel(Platform);
-	switch(Platform)
-	{
-		case SP_OPENGL_ES3_1_ANDROID:
-		case SP_OPENGL_ES31_EXT:
-		case SP_OPENGL_SM4:
-		case SP_OPENGL_SM5:
-		case SP_SWITCH:
-			OutDeclaration = CreateHLSLUniformBufferDeclaration(Name, UniformBufferStruct, MaxFeatureLevel, false);
-		case SP_PCD3D_SM5:
-		default:
-			OutDeclaration = CreateHLSLUniformBufferDeclaration(Name, UniformBufferStruct, MaxFeatureLevel, true);
-	}
+	OutDeclaration = CreateHLSLUniformBufferDeclaration(Name, UniformBufferStruct);
 }
 
-SHADERCORE_API void CacheUniformBufferIncludes(TMap<const TCHAR*,FCachedUniformBufferDeclaration>& Cache, EShaderPlatform Platform)
+SHADERCORE_API void CacheUniformBufferIncludes(TMap<const TCHAR*,FCachedUniformBufferDeclaration>& Cache)
 {
 	for (TMap<const TCHAR*,FCachedUniformBufferDeclaration>::TIterator It(Cache); It; ++It)
 	{
 		FCachedUniformBufferDeclaration& BufferDeclaration = It.Value();
-		check(BufferDeclaration.Declaration[Platform].Get() == NULL);
+		check(BufferDeclaration.Declaration.Get() == NULL);
 
 		for (TLinkedList<FUniformBufferStruct*>::TIterator StructIt(FUniformBufferStruct::GetStructList()); StructIt; StructIt.Next())
 		{
 			if (It.Key() == StructIt->GetShaderVariableName())
 			{
 				FString* NewDeclaration = new FString();
-				CreateUniformBufferShaderDeclaration(StructIt->GetShaderVariableName(), **StructIt, Platform, *NewDeclaration);
-				BufferDeclaration.Declaration[Platform] = MakeShareable(NewDeclaration);
+				CreateUniformBufferShaderDeclaration(StructIt->GetShaderVariableName(), **StructIt, *NewDeclaration);
+				check(!NewDeclaration->IsEmpty());
+				BufferDeclaration.Declaration = MakeShareable(NewDeclaration);
 				break;
 			}
 		}
@@ -361,21 +354,22 @@ SHADERCORE_API void CacheUniformBufferIncludes(TMap<const TCHAR*,FCachedUniformB
 void FShaderType::AddReferencedUniformBufferIncludes(FShaderCompilerEnvironment& OutEnvironment, FString& OutSourceFilePrefix, EShaderPlatform Platform)
 {
 	// Cache uniform buffer struct declarations referenced by this shader type's files
-	if (!bCachedUniformBufferStructDeclarations[Platform])
+	if (!bCachedUniformBufferStructDeclarations)
 	{
-		CacheUniformBufferIncludes(ReferencedUniformBufferStructsCache, Platform);
-		bCachedUniformBufferStructDeclarations[Platform] = true;
+		CacheUniformBufferIncludes(ReferencedUniformBufferStructsCache);
+		bCachedUniformBufferStructDeclarations = true;
 	}
 
 	FString UniformBufferIncludes;
 
 	for (TMap<const TCHAR*,FCachedUniformBufferDeclaration>::TIterator It(ReferencedUniformBufferStructsCache); It; ++It)
 	{
-		check(It.Value().Declaration[Platform].Get() != NULL);
+		check(It.Value().Declaration.Get() != NULL);
+		check(!It.Value().Declaration.Get()->IsEmpty());
 		UniformBufferIncludes += FString::Printf(TEXT("#include \"/Engine/Generated/UniformBuffers/%s.ush\"") LINE_TERMINATOR, It.Key());
 		OutEnvironment.IncludeVirtualPathToExternalContentsMap.Add(
 			*FString::Printf(TEXT("/Engine/Generated/UniformBuffers/%s.ush"),It.Key()),
-			It.Value().Declaration[Platform]
+			It.Value().Declaration
 			);
 
 		for (TLinkedList<FUniformBufferStruct*>::TIterator StructIt(FUniformBufferStruct::GetStructList()); StructIt; StructIt.Next())
@@ -389,6 +383,12 @@ void FShaderType::AddReferencedUniformBufferIncludes(FShaderCompilerEnvironment&
 
 	FString& GeneratedUniformBuffersInclude = OutEnvironment.IncludeVirtualPathToContentsMap.FindOrAdd("/Engine/Generated/GeneratedUniformBuffers.ush");
 	GeneratedUniformBuffersInclude.Append(UniformBufferIncludes);
+
+	ERHIFeatureLevel::Type MaxFeatureLevel = GetMaxSupportedFeatureLevel(Platform);
+	if (MaxFeatureLevel >= ERHIFeatureLevel::SM4)
+	{
+		OutEnvironment.SetDefine(TEXT("PLATFORM_SUPPORTS_SRV_UB"), TEXT("1"));
+	}
 }
 
 void FShaderType::DumpDebugInfo()
@@ -474,21 +474,22 @@ void FShaderType::SaveShaderStableKeys(EShaderPlatform TargetShaderPlatform)
 void FVertexFactoryType::AddReferencedUniformBufferIncludes(FShaderCompilerEnvironment& OutEnvironment, FString& OutSourceFilePrefix, EShaderPlatform Platform)
 {
 	// Cache uniform buffer struct declarations referenced by this shader type's files
-	if (!bCachedUniformBufferStructDeclarations[Platform])
+	if (!bCachedUniformBufferStructDeclarations)
 	{
-		CacheUniformBufferIncludes(ReferencedUniformBufferStructsCache, Platform);
-		bCachedUniformBufferStructDeclarations[Platform] = true;
+		CacheUniformBufferIncludes(ReferencedUniformBufferStructsCache);
+		bCachedUniformBufferStructDeclarations = true;
 	}
 
 	FString UniformBufferIncludes;
 
 	for (TMap<const TCHAR*,FCachedUniformBufferDeclaration>::TIterator It(ReferencedUniformBufferStructsCache); It; ++It)
 	{
-		check(It.Value().Declaration[Platform].Get() != NULL);
+		check(It.Value().Declaration.Get() != NULL);
+		check(!It.Value().Declaration.Get()->IsEmpty());
 		UniformBufferIncludes += FString::Printf(TEXT("#include \"/Engine/Generated/UniformBuffers/%s.ush\"") LINE_TERMINATOR, It.Key());
 		OutEnvironment.IncludeVirtualPathToExternalContentsMap.Add(
 			*FString::Printf(TEXT("/Engine/Generated/UniformBuffers/%s.ush"),It.Key()),
-			It.Value().Declaration[Platform]
+			It.Value().Declaration
 		);
 
 		for (TLinkedList<FUniformBufferStruct*>::TIterator StructIt(FUniformBufferStruct::GetStructList()); StructIt; StructIt.Next())
@@ -502,4 +503,10 @@ void FVertexFactoryType::AddReferencedUniformBufferIncludes(FShaderCompilerEnvir
 
 	FString& GeneratedUniformBuffersInclude = OutEnvironment.IncludeVirtualPathToContentsMap.FindOrAdd("/Engine/Generated/GeneratedUniformBuffers.ush");
 	GeneratedUniformBuffersInclude.Append(UniformBufferIncludes);
+
+	ERHIFeatureLevel::Type MaxFeatureLevel = GetMaxSupportedFeatureLevel(Platform);
+	if (MaxFeatureLevel >= ERHIFeatureLevel::SM4)
+	{
+		OutEnvironment.SetDefine(TEXT("PLATFORM_SUPPORTS_SRV_UB"), TEXT("1"));
+	}
 }
diff --git a/Engine/Source/Runtime/ShaderCore/Private/VertexFactory.cpp b/Engine/Source/Runtime/ShaderCore/Private/VertexFactory.cpp
--- a/Engine/Source/Runtime/ShaderCore/Private/VertexFactory.cpp
+++ b/Engine/Source/Runtime/ShaderCore/Private/VertexFactory.cpp
@@ -117,10 +117,7 @@ FVertexFactoryType::FVertexFactoryType(
 		TEXT("Incorrect virtual shader path extension for vertex factory shader header '%s': Only .ush files should be included."),
 		InShaderFilename);
 
-	for (int32 Platform = 0; Platform < SP_NumPlatforms; Platform++)
-	{
-		bCachedUniformBufferStructDeclarations[Platform] = false;
-	}
+	bCachedUniformBufferStructDeclarations = false;
 
 	// This will trigger if an IMPLEMENT_VERTEX_FACTORY_TYPE was in a module not loaded before InitializeShaderTypes
 	// Vertex factory types need to be implemented in modules that are loaded before that
diff --git a/Engine/Source/Runtime/ShaderCore/Public/Shader.h b/Engine/Source/Runtime/ShaderCore/Public/Shader.h
--- a/Engine/Source/Runtime/ShaderCore/Public/Shader.h
+++ b/Engine/Source/Runtime/ShaderCore/Public/Shader.h
@@ -1039,11 +1039,7 @@ public:
 	{
 		ReferencedUniformBufferStructsCache.Empty();
 		GenerateReferencedUniformBuffers(SourceFilename, Name, ShaderFileToUniformBufferVariables, ReferencedUniformBufferStructsCache);
-
-		for (int32 Platform = 0; Platform < SP_NumPlatforms; Platform++)
-		{
-			bCachedUniformBufferStructDeclarations[Platform] = false;
-		}
+		bCachedUniformBufferStructDeclarations = false;
 	}
 
 	void AddToShaderIdMap(FShaderId Id, FShader* Shader)
@@ -1104,7 +1100,7 @@ private:
 
 protected:
 	/** Tracks what platforms ReferencedUniformBufferStructsCache has had declarations cached for. */
-	bool bCachedUniformBufferStructDeclarations[SP_NumPlatforms];
+	bool bCachedUniformBufferStructDeclarations;
 
 	/**
 	* Cache of referenced uniform buffer includes.
diff --git a/Engine/Source/Runtime/ShaderCore/Public/ShaderCore.h b/Engine/Source/Runtime/ShaderCore/Public/ShaderCore.h
--- a/Engine/Source/Runtime/ShaderCore/Public/ShaderCore.h
+++ b/Engine/Source/Runtime/ShaderCore/Public/ShaderCore.h
@@ -1022,7 +1022,7 @@ extern SHADERCORE_API void VerifyShaderSourceFiles();
 struct FCachedUniformBufferDeclaration
 {
 	// Using SharedPtr so we can hand off lifetime ownership to FShaderCompilerEnvironment::IncludeVirtualPathToExternalContentsMap when invalidating this cache
-	TSharedPtr<FString> Declaration[SP_NumPlatforms];
+	TSharedPtr<FString> Declaration;
 };
 
 /** Parses the given source file and its includes for references of uniform buffers, which are then stored in UniformBufferEntries. */
diff --git a/Engine/Source/Runtime/ShaderCore/Public/ShaderParameters.h b/Engine/Source/Runtime/ShaderCore/Public/ShaderParameters.h
--- a/Engine/Source/Runtime/ShaderCore/Public/ShaderParameters.h
+++ b/Engine/Source/Runtime/ShaderCore/Public/ShaderParameters.h
@@ -157,7 +157,7 @@ private:
 };
 
 /** Creates a shader code declaration of this struct for the given shader platform. */
-extern SHADERCORE_API void CreateUniformBufferShaderDeclaration(const TCHAR* Name,const FUniformBufferStruct& UniformBufferStruct,EShaderPlatform Platform, FString& OutDeclaration);
+extern SHADERCORE_API void CreateUniformBufferShaderDeclaration(const TCHAR* Name,const FUniformBufferStruct& UniformBufferStruct, FString& OutDeclaration);
 
 class FShaderUniformBufferParameter
 {
diff --git a/Engine/Source/Runtime/ShaderCore/Public/VertexFactory.h b/Engine/Source/Runtime/ShaderCore/Public/VertexFactory.h
--- a/Engine/Source/Runtime/ShaderCore/Public/VertexFactory.h
+++ b/Engine/Source/Runtime/ShaderCore/Public/VertexFactory.h
@@ -216,11 +216,7 @@ public:
 	{
 		ReferencedUniformBufferStructsCache.Empty();
 		GenerateReferencedUniformBuffers(ShaderFilename, Name, ShaderFileToUniformBufferVariables, ReferencedUniformBufferStructsCache);
-
-		for (int32 Platform = 0; Platform < SP_NumPlatforms; Platform++)
-		{
-			bCachedUniformBufferStructDeclarations[Platform] = false;
-		}
+		bCachedUniformBufferStructDeclarations = false;
 	}
 
 	const TMap<const TCHAR*, FCachedUniformBufferDeclaration>& GetReferencedUniformBufferStructsCache() const
@@ -258,7 +254,7 @@ private:
 	TMap<const TCHAR*, FCachedUniformBufferDeclaration> ReferencedUniformBufferStructsCache;
 
 	/** Tracks what platforms ReferencedUniformBufferStructsCache has had declarations cached for. */
-	bool bCachedUniformBufferStructDeclarations[SP_NumPlatforms];
+	bool bCachedUniformBufferStructDeclarations;
 
 	/** 
 	 * Stores a history of serialization sizes for this vertex factory's shader parameter class. 
