Cherry-pick 4073144: Redo 3689687

diff --git a/Engine/Source/Runtime/ShaderCore/Public/Shader.h b/Engine/Source/Runtime/ShaderCore/Public/Shader.h
--- a/Engine/Source/Runtime/ShaderCore/Public/Shader.h
+++ b/Engine/Source/Runtime/ShaderCore/Public/Shader.h
@@ -1585,6 +1585,22 @@ public:
 
 	using FShaderPrimaryKey = TShaderTypePermutation<FShaderType>;
 
+	/** Used to compare two shader types by name. */
+	class FCompareShaderPrimaryKey
+	{
+	public:
+		FORCEINLINE bool operator()(const FShaderPrimaryKey& A, const FShaderPrimaryKey& B) const
+		{
+			int32 AL = FCString::Strlen(A.Type->GetName());
+			int32 BL = FCString::Strlen(B.Type->GetName());
+			if (AL == BL)
+			{
+				return FCString::Strncmp(A.Type->GetName(), B.Type->GetName(), AL) > 0 || A.PermutationId > B.PermutationId;
+			}
+			return AL > BL;
+		}
+	};
+
 	/** Default constructor. */
 	TShaderMap(EShaderPlatform InPlatform)
 		: Platform(InPlatform)
@@ -1805,16 +1821,17 @@ public:
 
 			// Sort the shaders by type name before saving, to make sure the saved result is binary equivalent to what is generated on other machines, 
 			// Which is a requirement of the Derived Data Cache.
-			auto SortedShaders = Shaders;
-			SortedShaders.KeySort(TCompareShaderTypePermutation<FShaderType>());
+			TArray<FShaderPrimaryKey> SortedShaderKeys;
+			Shaders.GenerateKeyArray(SortedShaderKeys);
+			SortedShaderKeys.Sort(FCompareShaderPrimaryKey());
 
-			for (typename TMap<FShaderPrimaryKey, TRefCountPtr<FShader> >::TIterator ShaderIt(SortedShaders); ShaderIt; ++ShaderIt)
+			for (FShaderPrimaryKey Key : SortedShaderKeys)
 			{
-				FShader* CurrentShader = ShaderIt.Value();
-				FShaderType* Type = ShaderIt.Key().Type;
+				FShaderType* Type = Key.Type;
 				check(Type);
 				checkSlow(FName(Type->GetName()) != NAME_None);
 				Ar << Type;
+				FShader* CurrentShader = Shaders.FindChecked(Key);
 				SerializeShaderForSaving(CurrentShader, Ar, bHandleShaderKeyChanges, bInlineShaderResource);
 			}
 
