Added ValidateCompiledResult to FShader. ValidateCompiledResult is currently used to check if we are trying to read from SceneTextureStruct from a base pass shader. This is possible to do from a custom shader node and before this check it would result in a runtime crash, as SceneTextureStruct isn't bound during the base pass.

diff --git a/Engine/Plugins/FX/Niagara/Source/NiagaraShader/Public/NiagaraShader.h b/Engine/Plugins/FX/Niagara/Source/NiagaraShader/Public/NiagaraShader.h
--- a/Engine/Plugins/FX/Niagara/Source/NiagaraShader/Public/NiagaraShader.h
+++ b/Engine/Plugins/FX/Niagara/Source/NiagaraShader/Public/NiagaraShader.h
@@ -58,6 +58,11 @@ public:
 	static void ModifyCompilationEnvironment(EShaderPlatform Platform, const FNiagaraShaderScript*  Script, FShaderCompilerEnvironment& OutEnvironment)
 	{
 	}
+
+	static bool ValidateCompiledResult(EShaderPlatform Platform, const FShaderParameterMap& ParameterMap, TArray<FString>& OutError)
+	{
+		return true;
+	}
 	
 	void SetDataInterfaceParameterInfo(const TArray<FNiagaraDataInterfaceGPUParamInfo>& InDIParamInfo);
 
diff --git a/Engine/Plugins/FX/Niagara/Source/NiagaraShader/Public/NiagaraShaderType.h b/Engine/Plugins/FX/Niagara/Source/NiagaraShader/Public/NiagaraShaderType.h
--- a/Engine/Plugins/FX/Niagara/Source/NiagaraShader/Public/NiagaraShaderType.h
+++ b/Engine/Plugins/FX/Niagara/Source/NiagaraShader/Public/NiagaraShaderType.h
@@ -68,6 +68,7 @@ public:
 	};
 	typedef FShader* (*ConstructCompiledType)(const CompiledShaderInitializerType&);
 	typedef bool (*ShouldCompilePermutationType)(EShaderPlatform,const FNiagaraShaderScript*);
+	typedef bool(*ValidateCompiledResultType)(EShaderPlatform, const FShaderParameterMap&, TArray<FString>&);
 	typedef void (*ModifyCompilationEnvironmentType)(EShaderPlatform,const FNiagaraShaderScript*, FShaderCompilerEnvironment&);
 
 	FNiagaraShaderType(
@@ -80,11 +81,13 @@ public:
 		ConstructCompiledType InConstructCompiledRef,
 		ModifyCompilationEnvironmentType InModifyCompilationEnvironmentRef,
 		ShouldCompilePermutationType InShouldCompilePermutationRef,
+		ValidateCompiledResultType InValidateCompiledResultRef,
 		GetStreamOutElementsType InGetStreamOutElementsRef
 		):
 		FShaderType(EShaderTypeForDynamicCast::Niagara, InName, InSourceFilename, InFunctionName, SF_Compute, InTotalPermutationCount, InConstructSerializedRef, InGetStreamOutElementsRef),
 		ConstructCompiledRef(InConstructCompiledRef),
 		ShouldCompilePermutationRef(InShouldCompilePermutationRef),
+		ValidateCompiledResultRef(InValidateCompiledResultRef),
 		ModifyCompilationEnvironmentRef(InModifyCompilationEnvironmentRef)
 	{
 		check(InTotalPermutationCount == 1);
@@ -148,6 +151,7 @@ protected:
 private:
 	ConstructCompiledType ConstructCompiledRef;
 	ShouldCompilePermutationType ShouldCompilePermutationRef;
+	ValidateCompiledResultType ValidateCompiledResultRef;
 	ModifyCompilationEnvironmentType ModifyCompilationEnvironmentRef;
 
 	static TMap<const FShaderCompileJob*, TArray<FNiagaraDataInterfaceGPUParamInfo> > ExtraParamInfo;
diff --git a/Engine/Source/Runtime/Engine/Private/ShaderCompiler/ShaderCompiler.cpp b/Engine/Source/Runtime/Engine/Private/ShaderCompiler/ShaderCompiler.cpp
--- a/Engine/Source/Runtime/Engine/Private/ShaderCompiler/ShaderCompiler.cpp
+++ b/Engine/Source/Runtime/Engine/Private/ShaderCompiler/ShaderCompiler.cpp
@@ -45,6 +45,7 @@
 #include "ProfilingDebugging/CookStats.h"
 #include "SceneInterface.h"
 #include "ShaderCodeLibrary.h"
+#include "MeshMaterialShaderType.h"
 
 #define LOCTEXT_NAMESPACE "ShaderCompiler"
 
@@ -798,7 +799,7 @@ void FShaderCompileUtilities::DoReadTaskResults(const TArray<FShaderCommonCompil
 	}
 }
 
-static void CheckSingleJob(FShaderCompileJob* SingleJob, TArray<FString>& Errors)
+static bool CheckSingleJob(FShaderCompileJob* SingleJob, const TArray<FMaterial*>& Materials, TArray<FString>& Errors)
 {
 	if (SingleJob->bSucceeded)
 	{
@@ -812,6 +813,39 @@ static void CheckSingleJob(FShaderCompileJob* SingleJob, TArray<FString>& Errors
 			Errors.AddUnique(SingleJob->Output.Errors[ErrorIndex].GetErrorString());
 		}
 	}
+
+	bool bSucceeded = SingleJob->bSucceeded;
+
+	if (SingleJob->ShaderType)
+	{
+		// Allow the shader validation to fail the compile if it sees any parameters bound that aren't supported.
+		if (FMaterialShaderType* MaterialShaderType = SingleJob->ShaderType->GetMaterialShaderType())
+		{
+			bSucceeded = bSucceeded && MaterialShaderType->ValidateCompiledResult(
+				(EShaderPlatform)SingleJob->Input.Target.Platform,
+				Materials,
+				SingleJob->Output.ParameterMap,
+				Errors);
+		}
+		else if (FMeshMaterialShaderType* MeshMaterialShaderType = SingleJob->ShaderType->GetMeshMaterialShaderType())
+		{
+			bSucceeded = bSucceeded && MeshMaterialShaderType->ValidateCompiledResult(
+				(EShaderPlatform)SingleJob->Input.Target.Platform,
+				Materials,
+				SingleJob->VFType,
+				SingleJob->Output.ParameterMap,
+				Errors);
+		}
+		else if (FGlobalShaderType* GlobalShaderType = SingleJob->ShaderType->GetGlobalShaderType())
+		{
+			bSucceeded = bSucceeded && GlobalShaderType->ValidateCompiledResult(
+				(EShaderPlatform)SingleJob->Input.Target.Platform,
+				SingleJob->Output.ParameterMap,
+				Errors);
+		}
+	}
+
+	return bSucceeded;
 };
 
 static void AddErrorsForFailedJob(const FShaderCompileJob& CurrentJob, TArray<EShaderPlatform>& ErrorPlatforms, TArray<FString>& UniqueErrors, TArray<const FShaderCommonCompileJob*>& ErrorJobs)
@@ -1933,7 +1967,7 @@ void FShaderCompilingManager::ProcessCompiledShaderMaps(
 				auto* SingleJob = CurrentJob.GetSingleShaderJob();
 				if (SingleJob)
 				{
-					CheckSingleJob(SingleJob, Errors);
+					bSuccess = bSuccess && CheckSingleJob(SingleJob, MaterialsArray, Errors);
 				}
 				else
 				{
@@ -1941,7 +1975,7 @@ void FShaderCompilingManager::ProcessCompiledShaderMaps(
 					for (int32 Index = 0; Index < PipelineJob->StageJobs.Num(); ++Index)
 					{
 						bSuccess = bSuccess && PipelineJob->StageJobs[Index]->bSucceeded;
-						CheckSingleJob(PipelineJob->StageJobs[Index]->GetSingleShaderJob(), Errors);
+						bSuccess = bSuccess && CheckSingleJob(PipelineJob->StageJobs[Index]->GetSingleShaderJob(), MaterialsArray, Errors);
 					}
 				}
 			}
diff --git a/Engine/Source/Runtime/Engine/Public/MaterialShaderType.h b/Engine/Source/Runtime/Engine/Public/MaterialShaderType.h
--- a/Engine/Source/Runtime/Engine/Public/MaterialShaderType.h
+++ b/Engine/Source/Runtime/Engine/Public/MaterialShaderType.h
@@ -87,6 +87,7 @@ public:
 
 	typedef FShader* (*ConstructCompiledType)(const CompiledShaderInitializerType&);
 	typedef bool (*ShouldCompilePermutationType)(EShaderPlatform,const FMaterial*);
+	typedef bool(*ValidateCompiledResultType)(EShaderPlatform, const TArray<FMaterial*>&, const FShaderParameterMap&, TArray<FString>&);
 	typedef void (*ModifyCompilationEnvironmentType)(EShaderPlatform, const FMaterial*, FShaderCompilerEnvironment&);
 
 	FMaterialShaderType(
@@ -99,11 +100,13 @@ public:
 		ConstructCompiledType InConstructCompiledRef,
 		ModifyCompilationEnvironmentType InModifyCompilationEnvironmentRef,
 		ShouldCompilePermutationType InShouldCompilePermutationRef,
+		ValidateCompiledResultType InValidateCompiledResultRef,
 		GetStreamOutElementsType InGetStreamOutElementsRef
 		):
 		FShaderType(EShaderTypeForDynamicCast::Material, InName, InSourceFilename, InFunctionName, InFrequency, InTotalPermutationCount, InConstructSerializedRef, InGetStreamOutElementsRef),
 		ConstructCompiledRef(InConstructCompiledRef),
 		ShouldCompilePermutationRef(InShouldCompilePermutationRef),
+		ValidateCompiledResultRef(InValidateCompiledResultRef),
 		ModifyCompilationEnvironmentRef(InModifyCompilationEnvironmentRef)
 	{
 		checkf(FPaths::GetExtension(InSourceFilename) == TEXT("usf"),
@@ -158,6 +161,18 @@ public:
 		return (*ShouldCompilePermutationRef)(Platform,Material);
 	}
 
+	/**
+	* Checks if the shader type should pass compilation for a particular set of parameters.
+	* @param Platform - Platform to validate.
+	* @param Materials - Materials to validate.
+	* @param ParameterMap - Shader parameters to validate.
+	* @param OutError - List for appending validation errors.
+	*/
+	bool ValidateCompiledResult(EShaderPlatform Platform, const TArray<FMaterial*>& Materials, const FShaderParameterMap& ParameterMap, TArray<FString>& OutError) const
+	{
+		return (*ValidateCompiledResultRef)(Platform, Materials, ParameterMap, OutError);
+	}
+
 protected:
 
 	/**
@@ -174,5 +189,6 @@ protected:
 private:
 	ConstructCompiledType ConstructCompiledRef;
 	ShouldCompilePermutationType ShouldCompilePermutationRef;
+	ValidateCompiledResultType ValidateCompiledResultRef;
 	ModifyCompilationEnvironmentType ModifyCompilationEnvironmentRef;
 };
diff --git a/Engine/Source/Runtime/Engine/Public/MeshMaterialShaderType.h b/Engine/Source/Runtime/Engine/Public/MeshMaterialShaderType.h
--- a/Engine/Source/Runtime/Engine/Public/MeshMaterialShaderType.h
+++ b/Engine/Source/Runtime/Engine/Public/MeshMaterialShaderType.h
@@ -60,6 +60,7 @@ public:
 	};
 	typedef FShader* (*ConstructCompiledType)(const CompiledShaderInitializerType&);
 	typedef bool (*ShouldCompilePermutationType)(EShaderPlatform,const FMaterial*,const FVertexFactoryType* VertexFactoryType);
+	typedef bool(*ValidateCompiledResultType)(EShaderPlatform, const TArray<FMaterial*>&, const FVertexFactoryType*, const FShaderParameterMap&, TArray<FString>&);
 	typedef void (*ModifyCompilationEnvironmentType)(EShaderPlatform, const FMaterial*, FShaderCompilerEnvironment&);
 
 	FMeshMaterialShaderType(
@@ -72,11 +73,13 @@ public:
 		ConstructCompiledType InConstructCompiledRef,
 		ModifyCompilationEnvironmentType InModifyCompilationEnvironmentRef,
 		ShouldCompilePermutationType InShouldCompilePermutationRef,
+		ValidateCompiledResultType InValidateCompiledResultRef,
 		GetStreamOutElementsType InGetStreamOutElementsRef
 		):
 		FShaderType(EShaderTypeForDynamicCast::MeshMaterial, InName,InSourceFilename,InFunctionName,InFrequency,InTotalPermutationCount,InConstructSerializedRef,InGetStreamOutElementsRef),
 		ConstructCompiledRef(InConstructCompiledRef),
 		ShouldCompilePermutationRef(InShouldCompilePermutationRef),
+		ValidateCompiledResultRef(InValidateCompiledResultRef),
 		ModifyCompilationEnvironmentRef(InModifyCompilationEnvironmentRef)
 	{
 		checkf(FPaths::GetExtension(InSourceFilename) == TEXT("usf"),
@@ -137,6 +140,19 @@ public:
 		return (*ShouldCompilePermutationRef)(Platform,Material,VertexFactoryType);
 	}
 
+	/**
+	* Checks if the shader type should pass compilation for a particular set of parameters.
+	* @param Platform - Platform to validate.
+	* @param Materials - Materials to validate.
+	* @param VertexFactoryType - Vertex factory to validate.
+	* @param ParameterMap - Shader parameters to validate.
+	* @param OutError - List for appending validation errors.
+	*/
+	bool ValidateCompiledResult(EShaderPlatform Platform, const TArray<FMaterial*>& Materials, const FVertexFactoryType* VertexFactoryType, const FShaderParameterMap& ParameterMap, TArray<FString>& OutError) const
+	{
+		return (*ValidateCompiledResultRef)(Platform, Materials, VertexFactoryType, ParameterMap, OutError);
+	}
+
 protected:
 
 	/**
@@ -153,6 +169,7 @@ protected:
 private:
 	ConstructCompiledType ConstructCompiledRef;
 	ShouldCompilePermutationType ShouldCompilePermutationRef;
+	ValidateCompiledResultType ValidateCompiledResultRef;
 	ModifyCompilationEnvironmentType ModifyCompilationEnvironmentRef;
 };
 
diff --git a/Engine/Source/Runtime/Renderer/Private/BasePassRendering.h b/Engine/Source/Runtime/Renderer/Private/BasePassRendering.h
--- a/Engine/Source/Runtime/Renderer/Private/BasePassRendering.h
+++ b/Engine/Source/Runtime/Renderer/Private/BasePassRendering.h
@@ -460,6 +460,17 @@ public:
 		FForwardLightingParameters::ModifyCompilationEnvironment(Platform, OutEnvironment);
 	}
 
+	static bool ValidateCompiledResult(EShaderPlatform Platform, const TArray<FMaterial*>& Materials, const FVertexFactoryType* VertexFactoryType, const FShaderParameterMap& ParameterMap, TArray<FString>& OutError)
+	{
+		if (ParameterMap.ContainsParameterAllocation(FSceneTexturesUniformParameters::StaticStruct.GetShaderVariableName()))
+		{
+			OutError.Add(TEXT("Base pass shaders cannot read from the SceneTexturesStruct."));
+			return false;
+		}
+
+		return true;
+	}
+
 	/** Initialization constructor. */
 	TBasePassPixelShaderPolicyParamType(const FMeshMaterialShaderType::CompiledShaderInitializerType& Initializer):
 		FMeshMaterialShader(Initializer)
diff --git a/Engine/Source/Runtime/Renderer/Public/MaterialShader.h b/Engine/Source/Runtime/Renderer/Public/MaterialShader.h
--- a/Engine/Source/Runtime/Renderer/Public/MaterialShader.h
+++ b/Engine/Source/Runtime/Renderer/Public/MaterialShader.h
@@ -113,6 +113,11 @@ public:
 	{
 	}
 
+	static bool ValidateCompiledResult(EShaderPlatform Platform, const TArray<FMaterial*>& Materials, const FShaderParameterMap& ParameterMap, TArray<FString>& OutError)
+	{
+		return true;
+	}
+
 	FUniformBufferRHIParamRef GetParameterCollectionBuffer(const FGuid& Id, const FSceneInterface* SceneInterface) const;
 
 	template<typename ShaderRHIParamRef>
diff --git a/Engine/Source/Runtime/Renderer/Public/MeshMaterialShader.h b/Engine/Source/Runtime/Renderer/Public/MeshMaterialShader.h
--- a/Engine/Source/Runtime/Renderer/Public/MeshMaterialShader.h
+++ b/Engine/Source/Runtime/Renderer/Public/MeshMaterialShader.h
@@ -32,6 +32,11 @@ public:
 		NonInstancedDitherLODFactorParameter.Bind(Initializer.ParameterMap, TEXT("NonInstancedDitherLODFactor"));
 	}
 
+	static bool ValidateCompiledResult(EShaderPlatform Platform, const TArray<FMaterial*>& Materials, const FVertexFactoryType* VertexFactoryType, const FShaderParameterMap& ParameterMap, TArray<FString>& OutError)
+	{
+		return true;
+	}
+
 	FORCEINLINE void ValidateAfterBind()
 	{
 		checkfSlow(PassUniformBuffer.IsInitialized(), TEXT("FMeshMaterialShader must bind a pass uniform buffer, even if it is just FSceneTexturesUniformParameters: %s"), GetType()->GetName());
diff --git a/Engine/Source/Runtime/ShaderCore/Public/GlobalShader.h b/Engine/Source/Runtime/ShaderCore/Public/GlobalShader.h
--- a/Engine/Source/Runtime/ShaderCore/Public/GlobalShader.h
+++ b/Engine/Source/Runtime/ShaderCore/Public/GlobalShader.h
@@ -63,6 +63,7 @@ public:
 	typedef FShader::CompiledShaderInitializerType CompiledShaderInitializerType;
 	typedef FShader* (*ConstructCompiledType)(const CompiledShaderInitializerType&);
 	typedef bool (*ShouldCompilePermutationType)(const FGlobalShaderPermutationParameters&);
+	typedef bool(*ValidateCompiledResultType)(EShaderPlatform, const FShaderParameterMap&, TArray<FString>&);
 	typedef void (*ModifyCompilationEnvironmentType)(const FGlobalShaderPermutationParameters&, FShaderCompilerEnvironment&);
 
 	FGlobalShaderType(
@@ -75,11 +76,13 @@ public:
 		ConstructCompiledType InConstructCompiledRef,
 		ModifyCompilationEnvironmentType InModifyCompilationEnvironmentRef,
 		ShouldCompilePermutationType InShouldCompilePermutationRef,
+		ValidateCompiledResultType InValidateCompiledResultRef,
 		GetStreamOutElementsType InGetStreamOutElementsRef
 		):
 		FShaderType(EShaderTypeForDynamicCast::Global, InName, InSourceFilename, InFunctionName, InFrequency, InTotalPermutationCount, InConstructSerializedRef, InGetStreamOutElementsRef),
 		ConstructCompiledRef(InConstructCompiledRef),
 		ShouldCompilePermutationRef(InShouldCompilePermutationRef),
+		ValidateCompiledResultRef(InValidateCompiledResultRef),
 		ModifyCompilationEnvironmentRef(InModifyCompilationEnvironmentRef)
 	{
 		checkf(FPaths::GetExtension(InSourceFilename) == TEXT("usf"),
@@ -109,9 +112,21 @@ public:
 		(*ModifyCompilationEnvironmentRef)(FGlobalShaderPermutationParameters(Platform, PermutationId), Environment);
 	}
 
+	/**
+	* Checks if the shader type should pass compilation for a particular set of parameters.
+	* @param Platform - Platform to validate.
+	* @param ParameterMap - Shader parameters to validate.
+	* @param OutError - List for appending validation errors.
+	*/
+	bool ValidateCompiledResult(EShaderPlatform Platform, const FShaderParameterMap& ParameterMap, TArray<FString>& OutError) const
+	{
+		return (*ValidateCompiledResultRef)(Platform, ParameterMap, OutError);
+	}
+
 private:
 	ConstructCompiledType ConstructCompiledRef;
 	ShouldCompilePermutationType ShouldCompilePermutationRef;
+	ValidateCompiledResultType ValidateCompiledResultRef;
 	ModifyCompilationEnvironmentType ModifyCompilationEnvironmentRef;
 };
 
@@ -143,6 +158,7 @@ public:
 
 	// FShader interface.
 	SHADERCORE_API static void ModifyCompilationEnvironment(const FGlobalShaderPermutationParameters& Parameters, FShaderCompilerEnvironment& OutEnvironment) {}
+	SHADERCORE_API static bool ValidateCompiledResult(EShaderPlatform Platform, const FShaderParameterMap& ParameterMap, TArray<FString>& OutErrors) { return true; }
 };
 
 /**
@@ -238,6 +254,10 @@ inline TShaderMap<FGlobalShaderType>* GetGlobalShaderMap(ERHIFeatureLevel::Type
 		return ShaderClass::ShouldCompilePermutation(Parameters); \
 	} \
 	\
+	static bool ValidateCompiledResultImpl(EShaderPlatform Platform, const FShaderParameterMap& ParameterMap, TArray<FString>& OutErrors) \
+	{ \
+		return ShaderClass::ValidateCompiledResult(Platform, ParameterMap, OutErrors); \
+	} \
 	static void ModifyCompilationEnvironmentImpl( \
 		const FGlobalShaderPermutationParameters& Parameters, \
 		FShaderCompilerEnvironment& OutEnvironment) \
@@ -258,5 +278,6 @@ inline TShaderMap<FGlobalShaderType>* GetGlobalShaderMap(ERHIFeatureLevel::Type
 		ShaderClass::ConstructCompiledInstance, \
 		ShaderClass::ModifyCompilationEnvironmentImpl, \
 		ShaderClass::ShouldCompilePermutationImpl, \
+		ShaderClass::ValidateCompiledResultImpl, \
 		ShaderClass::GetStreamOutElements \
 		)
diff --git a/Engine/Source/Runtime/ShaderCore/Public/Shader.h b/Engine/Source/Runtime/ShaderCore/Public/Shader.h
--- a/Engine/Source/Runtime/ShaderCore/Public/Shader.h
+++ b/Engine/Source/Runtime/ShaderCore/Public/Shader.h
@@ -1147,6 +1147,7 @@ protected:
 		ShaderClass::ConstructCompiledInstance, \
 		ShaderClass::ModifyCompilationEnvironment, \
 		ShaderClass::ShouldCompilePermutation, \
+		ShaderClass::ValidateCompiledResult, \
 		ShaderClass::GetStreamOutElements \
 		);
 
@@ -1163,6 +1164,7 @@ protected:
 		ShaderClass::ConstructCompiledInstance, \
 		ShaderClass::ModifyCompilationEnvironment, \
 		ShaderClass::ShouldCompilePermutation, \
+		ShaderClass::ValidateCompiledResult, \
 		ShaderClass::GetStreamOutElements \
 		);
 
@@ -1179,6 +1181,7 @@ protected:
 	ShaderClass::ConstructCompiledInstance, \
 	ShaderClass::ModifyCompilationEnvironment, \
 	ShaderClass::ShouldCompilePermutation, \
+	ShaderClass::ValidateCompiledResult, \
 	ShaderClass::GetStreamOutElements \
 	);
 
@@ -1195,6 +1198,7 @@ protected:
 	ShaderClass::ConstructCompiledInstance, \
 	ShaderClass::ModifyCompilationEnvironment, \
 	ShaderClass::ShouldCompilePermutation, \
+	ShaderClass::ValidateCompiledResult, \
 	ShaderClass::GetStreamOutElements \
 	);
 #endif
