Back out revision 3 from //UE4/Dev-Niagara/Engine/Source/Runtime/ShaderCore/Public/ShaderParameterUtils.h as it incorrectly stripped out an inline

#tests n/a

diff --git a/Engine/Source/Runtime/ShaderCore/Public/ShaderParameterUtils.h b/Engine/Source/Runtime/ShaderCore/Public/ShaderParameterUtils.h
--- a/Engine/Source/Runtime/ShaderCore/Public/ShaderParameterUtils.h
+++ b/Engine/Source/Runtime/ShaderCore/Public/ShaderParameterUtils.h
@@ -188,7 +188,7 @@ inline void SetPixelShaderBool(
  * Sets the value of a shader texture parameter.  Template'd on shader type
  */
 template<typename ShaderTypeRHIParamRef, typename TRHICmdList>
- void SetTextureParameter(
+FORCEINLINE void SetTextureParameter(
 	TRHICmdList& RHICmdList,
 	ShaderTypeRHIParamRef Shader,
 	const FShaderResourceParameter& TextureParameter,
