Fixing CIS

diff --git a/Engine/Source/Runtime/ShaderCore/Private/GlobalShader.cpp b/Engine/Source/Runtime/ShaderCore/Private/GlobalShader.cpp
--- a/Engine/Source/Runtime/ShaderCore/Private/GlobalShader.cpp
+++ b/Engine/Source/Runtime/ShaderCore/Private/GlobalShader.cpp
@@ -183,7 +183,7 @@ void BackupGlobalShaderMap(FGlobalShaderBackupData& OutGlobalShaderBackup)
 		{
 			TUniquePtr<TArray<uint8>> ShaderData = MakeUnique<TArray<uint8>>();
 			FMemoryWriter Ar(*ShaderData);
-			GGlobalShaderMap[ShaderPlatform]->SerializeInline(Ar, true, true, nullptr);
+			GGlobalShaderMap[ShaderPlatform]->SerializeInline(Ar, true, true, false, nullptr);
 			GGlobalShaderMap[ShaderPlatform]->RegisterSerializedShaders(false);
 			GGlobalShaderMap[ShaderPlatform]->Empty();
 			OutGlobalShaderBackup.FeatureLevelShaderData[i] = MoveTemp(ShaderData);
@@ -207,7 +207,7 @@ void RestoreGlobalShaderMap(const FGlobalShaderBackupData& GlobalShaderBackup)
 			&& GGlobalShaderMap[ShaderPlatform] != nullptr)
 		{
 			FMemoryReader Ar(*GlobalShaderBackup.FeatureLevelShaderData[i]);
-			GGlobalShaderMap[ShaderPlatform]->SerializeInline(Ar, true, true, nullptr);
+			GGlobalShaderMap[ShaderPlatform]->SerializeInline(Ar, true, true, false, nullptr);
 			GGlobalShaderMap[ShaderPlatform]->RegisterSerializedShaders(false);
 		}
 	}
