UE4.21 - Change message level as it's purely informational
#rnx

diff --git a/Engine/Source/Runtime/ShaderCore/Private/Shader.cpp b/Engine/Source/Runtime/ShaderCore/Private/Shader.cpp
--- a/Engine/Source/Runtime/ShaderCore/Private/Shader.cpp
+++ b/Engine/Source/Runtime/ShaderCore/Private/Shader.cpp
@@ -301,7 +301,7 @@ FArchive& operator<<(FArchive& Ar,FShaderType*& Ref)
 			}
 			else
 			{
-				UE_LOG(LogShaders, Warning, TEXT("ShaderType '%s' was not found!"), *ShaderTypeName.ToString());
+				UE_LOG(LogShaders, Verbose, TEXT("ShaderType '%s' dependency was not found."), *ShaderTypeName.ToString());
 			}
 		}
 	}
