REL - Revert CL 4307449. The issue is that a commandlet is used to prefill the DDC and the entire shadermap (all shader permuations including the debugshaders) is cached in the DDC key. If the key does not match because the editor expected debugviewmode shaders the entire shadermap is recreated and compile efectively rendering the commandlet ineffective.

diff --git a/Engine/Source/Runtime/ShaderCore/Private/ShaderCore.cpp b/Engine/Source/Runtime/ShaderCore/Private/ShaderCore.cpp
--- a/Engine/Source/Runtime/ShaderCore/Private/ShaderCore.cpp
+++ b/Engine/Source/Runtime/ShaderCore/Private/ShaderCore.cpp
@@ -282,7 +282,7 @@ bool AllowDebugViewmodes(EShaderPlatform Platform)
 	bool bForceEnable = ForceDebugViewValue == 1;
 	bool bForceDisable = ForceDebugViewValue == 2;
 	ITargetPlatform* TargetPlatform = GetTargetPlatformManager()->FindTargetPlatform(ShaderPlatformToPlatformName(Platform).ToString());
-	return (!bForceDisable) && (bForceEnable || !TargetPlatform || (!IsRunningCommandlet() && !TargetPlatform->RequiresCookedData()));
+	return (!bForceDisable) && (bForceEnable || !TargetPlatform || !TargetPlatform->RequiresCookedData());
 #else
 	return AllowDebugViewmodes();
 #endif
