Decouple ShaderPipelineCache from RenderCore.

diff --git a/Engine/Source/Runtime/Engine/Public/Rendering/ShaderPipelineStateCache.h b/Engine/Source/Runtime/Engine/Public/Rendering/ShaderPipelineStateCache.h
new file mode 100644
--- /dev/null
+++ b/Engine/Source/Runtime/Engine/Public/Rendering/ShaderPipelineStateCache.h
@@ -0,0 +1,62 @@
+// Copyright 1998-2018 Epic Games, Inc. All Rights Reserved.
+
+#pragma once
+
+#include "CoreMinimal.h"
+#include "Stats/Stats.h"
+#include "TickableObjectRenderThread.h"
+#include "ShaderPipelineCache.h"
+
+/**
+ * Wrapper class to decouple FShaderPipelineCache from RenderCore.
+ */
+class ENGINE_API FShaderPipelineStateCache : public FShaderPipelineCache, public FTickableObjectRenderThread
+{
+public:
+	/**
+	 * Initializes the shader pipeline cache for the desired platform, called by the engine.
+	 * The shader platform is used only to load the initial pipeline cache and can be changed by closing & reopening the cache if necessary.
+	 */
+	static void Initialize(EShaderPlatform Platform)
+	{
+		FShaderPipelineCache::Initialize(Platform);
+	}
+	
+	/** Terminates the shader pipeline cache, called by the engine. */
+	static void Shutdown()
+	{
+		FShaderPipelineCache::Shutdown();
+	}
+
+	FShaderPipelineStateCache(EShaderPlatform Platform)
+	: FShaderPipelineCache(Platform)
+	, FTickableObjectRenderThread(true, false) // (RegisterNow, HighFrequency)
+	{
+		
+	}
+	
+	virtual ~FShaderPipelineStateCache()
+	{
+		
+	}
+
+	virtual bool IsTickable() const override final
+	{
+		return FShaderPipelineCache::IsTickable();
+	}
+	
+	virtual void Tick( float DeltaTime ) override final
+	{
+		return FShaderPipelineCache::Tick(DeltaTime);
+	}
+	
+	virtual bool NeedsRenderingResumedForRenderingThreadTick() const override final
+	{
+		return FShaderPipelineCache::NeedsRenderingResumedForRenderingThreadTick();
+	}
+	
+	virtual TStatId GetStatId() const override final
+	{
+		return FShaderPipelineCache::GetStatId();
+	}
+};
diff --git a/Engine/Source/Runtime/Launch/Private/LaunchEngineLoop.cpp b/Engine/Source/Runtime/Launch/Private/LaunchEngineLoop.cpp
--- a/Engine/Source/Runtime/Launch/Private/LaunchEngineLoop.cpp
+++ b/Engine/Source/Runtime/Launch/Private/LaunchEngineLoop.cpp
@@ -153,7 +153,7 @@
     #include "PreLoadScreenManager.h"
 
 	#include "ShaderCodeLibrary.h"
-	#include "ShaderPipelineCache.h"
+	#include "Rendering/ShaderPipelineStateCache.h"
 
 #if !UE_BUILD_SHIPPING
 	#include "STaskGraph.h"
@@ -1908,7 +1908,7 @@ int32 FEngineLoop::PreInit(const TCHAR* CmdLine)
 			
 		// Initialize the pipeline cache system. Opening is deferred until the manual call to
 		// OpenPipelineFileCache below, after content pak's ShaderCodeLibraries are loaded.
-		FShaderPipelineCache::Initialize(GMaxRHIShaderPlatform);
+		FShaderPipelineStateCache::Initialize(GMaxRHIShaderPlatform);
 	}
 
 	FString Commandline = FCommandLine::Get();
@@ -3196,7 +3196,7 @@ void FEngineLoop::Exit()
 	StopRenderingThread();
 
 	// Disable the PSO cache
-	FShaderPipelineCache::Shutdown();
+	FShaderPipelineStateCache::Shutdown();
 
 	// Close shader code map, if any
 	FShaderCodeLibrary::Shutdown();
diff --git a/Engine/Source/Runtime/ShaderCore/Private/ShaderPipelineCache.cpp b/Engine/Source/Runtime/ShaderCore/Private/ShaderPipelineCache.cpp
--- a/Engine/Source/Runtime/ShaderCore/Private/ShaderPipelineCache.cpp
+++ b/Engine/Source/Runtime/ShaderCore/Private/ShaderPipelineCache.cpp
@@ -6,7 +6,6 @@
  */
  
 #include "ShaderPipelineCache.h"
-#include "RenderingThread.h"
 #include "HAL/FileManager.h"
 #include "Misc/Paths.h"
 #include "Serialization/MemoryWriter.h"
@@ -21,7 +20,6 @@
 #include "Misc/ScopeRWLock.h"
 #include "Misc/CoreDelegates.h"
 #include "ShaderCodeLibrary.h"
-#include "TickableObjectRenderThread.h"
 #include "Misc/ConfigCacheIni.h"
 #include "Async/AsyncFileHandle.h"
 
@@ -954,8 +952,7 @@ void FShaderPipelineCache::Flush()
 }
 
 FShaderPipelineCache::FShaderPipelineCache(EShaderPlatform Platform)
-: FTickableObjectRenderThread(true, false) // (RegisterNow, HighFrequency)
-, BatchSize(0)
+: BatchSize(0)
 , BatchTime(0.0f)
 , bPaused(false)
 , TotalActiveTasks(0)
diff --git a/Engine/Source/Runtime/ShaderCore/Public/ShaderPipelineCache.h b/Engine/Source/Runtime/ShaderCore/Public/ShaderPipelineCache.h
--- a/Engine/Source/Runtime/ShaderCore/Public/ShaderPipelineCache.h
+++ b/Engine/Source/Runtime/ShaderCore/Public/ShaderPipelineCache.h
@@ -10,7 +10,6 @@
 #include "CoreMinimal.h"
 #include "Stats/Stats.h"
 #include "RHI.h"
-#include "TickableObjectRenderThread.h"
 #include "PipelineFileCache.h"
 #include "Delegates/DelegateCombinations.h"
 
@@ -65,15 +64,15 @@ class FShaderPipelineCacheArchive;
  * - FShaderCodeLibrary must be enabled via Project Settings > Packaging > "Share Material Shader Code".
  * - Enabling "Native Shader Libraries" is optional, but strongly preferred for Metal.
  */
-class SHADERCORE_API FShaderPipelineCache : public FTickableObjectRenderThread
+class SHADERCORE_API FShaderPipelineCache
 {
 	struct CompileJob
 	{
 		FPipelineCacheFileFormatPSO PSO;
 		FShaderPipelineCacheArchive* ReadRequests;
 	};
-
-public:
+	
+protected:
 	/**
 	 * Initializes the shader pipeline cache for the desired platform, called by the engine.
 	 * The shader platform is used only to load the initial pipeline cache and can be changed by closing & reopening the cache if necessary.
@@ -81,7 +80,8 @@ public:
 	static void Initialize(EShaderPlatform Platform);
 	/** Terminates the shader pipeline cache, called by the engine. */
 	static void Shutdown();
-	
+
+public:
 	/** Pauses precompilation. */
 	static void PauseBatching();
 	
@@ -120,13 +120,13 @@ public:
 	FShaderPipelineCache(EShaderPlatform Platform);
 	virtual ~FShaderPipelineCache();
 
-	virtual bool IsTickable() const override final;
+	bool IsTickable() const;
 	
-	virtual void Tick( float DeltaTime ) override final;
+	void Tick( float DeltaTime );
 	
-	virtual bool NeedsRenderingResumedForRenderingThreadTick() const override final;
+	bool NeedsRenderingResumedForRenderingThreadTick() const;
 	
-	virtual TStatId GetStatId() const override final;
+	TStatId GetStatId() const;
 	
 	enum ELibraryState
 	{
