Reduced excessive memory usage at the end of a cook when bShareMaterialShaderCode is enabled. Savings vary by the total number of shaders - a shader library with 1.3m shaders will use approx 6GB less memory.

diff --git a/Engine/Source/Runtime/ShaderCore/Private/ShaderCodeLibrary.cpp b/Engine/Source/Runtime/ShaderCore/Private/ShaderCodeLibrary.cpp
--- a/Engine/Source/Runtime/ShaderCore/Private/ShaderCodeLibrary.cpp
+++ b/Engine/Source/Runtime/ShaderCore/Private/ShaderCodeLibrary.cpp
@@ -232,43 +232,46 @@ void FStableShaderKeyAndValue::ParseFromString(const FString& Src)
 FString FStableShaderKeyAndValue::ToString() const
 {
 	FString Result;
-	Result.Reserve(2048);
-
-	const FString Delim(TEXT(","));
-
-	Result += ClassNameAndObjectPath.ToString().Replace(TEXT(","), TEXT(" "));
-	Result += Delim;
-
-	Result += ShaderType.ToString().Replace(TEXT(","), TEXT(" "));
-	Result += Delim;
-	Result += ShaderClass.ToString().Replace(TEXT(","), TEXT(" "));
-	Result += Delim;
-	Result += MaterialDomain.ToString();
-	Result += Delim;
-	Result += FeatureLevel.ToString();
-	Result += Delim;
-
-	Result += QualityLevel.ToString();
-	Result += Delim;
-	Result += TargetFrequency.ToString();
-	Result += Delim;
-	Result += TargetPlatform.ToString();
-	Result += Delim;
-
-	Result += VFType.ToString();
-	Result += Delim;
-	Result += PermutationId.ToString();
-	Result += Delim;
-
-	Result += OutputHash.ToString();
-
+	ToString(Result);
 	return Result;
 }
 
+void FStableShaderKeyAndValue::ToString(FString& OutResult) const
+{
+	const TCHAR* Delim = TEXT(",");
+
+	OutResult.Reset(255);
+
+	OutResult += ClassNameAndObjectPath.ToString().Replace(Delim, TEXT(" "));
+	OutResult += Delim;
+
+	OutResult += ShaderType.ToString().Replace(Delim, TEXT(" "));
+	OutResult += Delim;
+	OutResult += ShaderClass.ToString().Replace(Delim, TEXT(" "));
+	OutResult += Delim;
+	OutResult += MaterialDomain.ToString();
+	OutResult += Delim;
+	OutResult += FeatureLevel.ToString();
+	OutResult += Delim;
+
+	OutResult += QualityLevel.ToString();
+	OutResult += Delim;
+	OutResult += TargetFrequency.ToString();
+	OutResult += Delim;
+	OutResult += TargetPlatform.ToString();
+	OutResult += Delim;
+
+	OutResult += VFType.ToString();
+	OutResult += Delim;
+	OutResult += PermutationId.ToString();
+	OutResult += Delim;
+
+	OutResult += OutputHash.ToString();
+}
+
 FString FStableShaderKeyAndValue::HeaderLine()
 {
 	FString Result;
-	Result.Reserve(2048);
 
 	const FString Delim(TEXT(","));
 
@@ -1250,12 +1253,14 @@ struct FEditorShaderStableInfo
 			FString IntermediateFormatPath = GetStableInfoArchiveFilename(FPaths::ProjectSavedDir() / TEXT("Shaders") / FormatName.ToString(), LibraryName, FormatName);
 
 			TArray<FString> FileContents;
+			FileContents.Reserve(StableMap.Num() + 1);
 
 			FileContents.Add(FStableShaderKeyAndValue::HeaderLine());
 
 			for (const FStableShaderKeyAndValue& Item : StableMap)
 			{
-				FileContents.Add(Item.ToString());
+				FString& LineBuffer = FileContents.Emplace_GetRef();
+				Item.ToString(LineBuffer);
 			}
 			FFileHelper::SaveStringArrayToFile(FileContents, *IntermediateFormatPath);
 
diff --git a/Engine/Source/Runtime/ShaderCore/Public/ShaderCodeLibrary.h b/Engine/Source/Runtime/ShaderCore/Public/ShaderCodeLibrary.h
--- a/Engine/Source/Runtime/ShaderCore/Public/ShaderCodeLibrary.h
+++ b/Engine/Source/Runtime/ShaderCore/Public/ShaderCodeLibrary.h
@@ -88,6 +88,7 @@ struct SHADERCORE_API FStableShaderKeyAndValue
 	void ComputeKeyHash();
 	void ParseFromString(const FString& Src);
 	FString ToString() const;
+	void ToString(FString& OutResult) const;
 	static FString HeaderLine();
 
 	friend bool operator ==(const FStableShaderKeyAndValue& A, const FStableShaderKeyAndValue& B)
